[ColorGradingEditor] [DisplayCluster] Separate color grading panel into its own plugin and split color grading/ICVFX drawers in DisplayCluster

diff --git a/Engine/Plugins/Editor/ColorGrading/Source/ColorGradingEditor/Private/ColorGradingListItem.cpp b/Engine/Plugins/Editor/ColorGrading/Source/ColorGradingEditor/Private/ColorGradingListItem.cpp
new file mode 100644
--- /dev/null
+++ b/Engine/Plugins/Editor/ColorGrading/Source/ColorGradingEditor/Private/ColorGradingListItem.cpp
@@ -0,0 +1,23 @@
+// Copyright Epic Games, Inc. All Rights Reserved.
+
+#include "ColorGradingListItem.h"
+
+TArray<FColorGradingListItemGeneratorRef> FColorGradingListItem::RegisteredListItemGenerators;
+TSet<TSubclassOf<AActor>> FColorGradingListItem::ActorClassesWithListItemGenerators;
+
+TArray<FColorGradingListItemRef> FColorGradingListItem::GenerateColorGradingListItems(AActor* InActor)
+{
+	TArray<FColorGradingListItemRef> ListItems;
+
+	for (const FColorGradingListItemGeneratorRef& GeneratorInstance : RegisteredListItemGenerators)
+	{
+		if (!GeneratorInstance.IsValid())
+		{
+			continue;
+		}
+
+		GeneratorInstance->GenerateColorGradingListItems(InActor, ListItems);
+	}
+
+	return ListItems;
+}
\ No newline at end of file
diff --git a/Engine/Plugins/Editor/ColorGrading/Source/ColorGradingEditor/Private/SColorGradingColorWheel.cpp b/Engine/Plugins/Editor/ColorGrading/Source/ColorGradingEditor/Private/SColorGradingColorWheel.cpp
new file mode 100644
--- /dev/null
+++ b/Engine/Plugins/Editor/ColorGrading/Source/ColorGradingEditor/Private/SColorGradingColorWheel.cpp
@@ -0,0 +1,845 @@
+// Copyright Epic Games, Inc. All Rights Reserved.
+
+#include "SColorGradingColorWheel.h"
+
+#include "Customizations/MathStructCustomizations.h"
+#include "DetailLayoutBuilder.h"
+#include "Editor.h"
+#include "PropertyHandle.h"
+#include "ScopedTransaction.h"
+#include "Widgets/SBoxPanel.h"
+#include "Widgets/Colors/SComplexGradient.h"
+#include "Widgets/Layout/SBox.h"
+
+#define LOCTEXT_NAMESPACE "ColorGradingEditor"
+
+void SColorGradingColorWheel::Construct(const FArguments& InArgs)
+{
+	Orientation = InArgs._Orientation;
+	ColorDisplayMode = InArgs._ColorDisplayMode;
+
+	ChildSlot
+	[
+		SNew(SVerticalBox)
+
+		+ SVerticalBox::Slot()
+		.AutoHeight()
+		.HAlign(HAlign_Fill)
+		[
+			SAssignNew(HeaderBox, SBox)
+		]
+
+		+ SVerticalBox::Slot()
+		.FillHeight(1.0f)
+		[
+			SAssignNew(ColorWheelBox, SBox)
+		]
+	];
+
+	ColorPickerBox = SNew(SBox);
+	ColorSlidersBox = SNew(SBox);
+
+	SetOrientation(Orientation);
+
+	if (InArgs._HeaderContent.IsValid())
+	{
+		HeaderBox->SetContent(InArgs._HeaderContent.ToSharedRef());
+	}
+}
+
+void SColorGradingColorWheel::SetColorPropertyHandle(TSharedPtr<IPropertyHandle> InColorPropertyHandle)
+{
+	ColorGradingPicker.Reset();
+	ColorPropertyMetadata.Reset();
+	ComponentSliderDynamicMinValue.Reset();
+	ComponentSliderDynamicMaxValue.Reset();
+
+	ColorPropertyHandle = InColorPropertyHandle;
+
+	if (ColorPropertyHandle.IsValid() && ColorPropertyHandle.Pin()->IsValidHandle())
+	{
+		ColorPropertyMetadata = GetColorPropertyMetadata();
+		ComponentSliderDynamicMinValue = ColorPropertyMetadata->MinValue;
+		ComponentSliderDynamicMaxValue = ColorPropertyMetadata->MaxValue;
+	}
+
+	// Since many of the color picker slate properties are not attributes, we need to recreate the widget every time the property handle changes
+	// to ensure the picker is configured correctly for the color property
+	if (ColorPickerBox.IsValid())
+	{
+		ColorPickerBox->SetContent(CreateColorGradingPicker());
+	}
+
+	if (ColorSlidersBox.IsValid())
+	{
+		ColorSlidersBox->SetContent(CreateColorComponentSliders());
+	}
+}
+
+void SColorGradingColorWheel::SetHeaderContent(const TSharedRef<SWidget>& HeaderContent)
+{
+	if (HeaderBox.IsValid())
+	{
+		HeaderBox->SetContent(HeaderContent);
+	}
+}
+
+void SColorGradingColorWheel::SetOrientation(EOrientation NewOrientation)
+{
+	Orientation = NewOrientation;
+
+	if (ColorWheelBox.IsValid())
+	{
+		TSharedPtr<SWidget> OrientedBox;
+
+		if (Orientation == EOrientation::Orient_Vertical)
+		{
+			OrientedBox = SNew(SVerticalBox)
+
+			+ SVerticalBox::Slot()
+			.AutoHeight()
+			[
+				ColorPickerBox.ToSharedRef()
+			]
+
+			+ SVerticalBox::Slot()
+			.AutoHeight()
+			[
+				ColorSlidersBox.ToSharedRef()
+			];
+		}
+		else
+		{
+			OrientedBox = SNew(SHorizontalBox)
+
+			+ SHorizontalBox::Slot()
+			.AutoWidth()
+			[
+				ColorPickerBox.ToSharedRef()
+			]
+
+			+ SHorizontalBox::Slot()
+			.FillWidth(1.0f)
+			[
+				ColorSlidersBox.ToSharedRef()
+			];
+		}
+
+		ColorWheelBox->SetContent(OrientedBox.ToSharedRef());
+		ColorWheelBox->SetHAlign(Orientation == EOrientation::Orient_Vertical ? HAlign_Center : HAlign_Fill);
+	}
+}
+
+TSharedRef<SWidget> SColorGradingColorWheel::SColorGradingColorWheel::CreateColorGradingPicker()
+{
+	if (ColorPropertyHandle.IsValid() && ColorPropertyHandle.Pin()->IsValidHandle())
+	{
+		return SNew(SColorGradingPicker)
+			.ValueMin(GetMetadataMinValue())
+			.ValueMax(GetMetadataMaxValue())
+			.SliderValueMin(GetMetadataSliderMinValue())
+			.SliderValueMax(GetMetadataSliderMaxValue())
+			.MainDelta(GetMetadataDelta())
+			.SupportDynamicSliderMinValue(GetMetadataSupportDynamicSliderMinValue())
+			.SupportDynamicSliderMaxValue(GetMetadataSupportDynamicSliderMaxValue())
+			.MainShiftMultiplier(GetMetadataShiftMultiplier())
+			.MainCtrlMultiplier(GetMetadataCtrlMultiplier())
+			.ColorGradingModes(ColorPropertyMetadata->ColorGradingMode)
+			.OnColorCommitted(this, &SColorGradingColorWheel::CommitColor)
+			.OnQueryCurrentColor(this, &SColorGradingColorWheel::GetColor)
+			.AllowSpin(ColorPropertyHandle.Pin()->GetNumOuterObjects() == 1) // TODO: May want to find a way to support multiple objects
+			.OnBeginSliderMovement(this, &SColorGradingColorWheel::BeginUsingColorPickerSlider)
+			.OnEndSliderMovement(this, &SColorGradingColorWheel::EndUsingColorPickerSlider)
+			.OnBeginMouseCapture(this, &SColorGradingColorWheel::BeginUsingColorPickerSlider)
+			.OnEndMouseCapture(this, &SColorGradingColorWheel::EndUsingColorPickerSlider)
+			.IsEnabled(this, &SColorGradingColorWheel::IsPropertyEnabled);
+	}
+
+	return SNullWidget::NullWidget;
+}
+
+TSharedRef<SWidget> SColorGradingColorWheel::CreateColorComponentSliders()
+{
+	if (ColorPropertyHandle.IsValid() && ColorPropertyHandle.Pin()->IsValidHandle())
+	{
+		TSharedRef<SVerticalBox> ColorSlidersVerticalBox = SNew(SVerticalBox);
+
+		uint32 NumComponents = 4;
+		for (uint32 ComponentIndex = 0; ComponentIndex < NumComponents; ++ComponentIndex)
+		{
+			TAttribute<FText> TextGetter = TAttribute<FText>::Create(TAttribute<FText>::FGetter::CreateSP(this, &SColorGradingColorWheel::GetComponentLabelText, ComponentIndex));
+			TSharedRef<SWidget> LabelWidget = SNumericEntryBox<float>::BuildLabel(TextGetter, FLinearColor::White, FLinearColor(0.2f, 0.2f, 0.2f));
+
+			TSharedRef<SNumericEntryBox<float>> NumericEntryBox = SNew(SNumericEntryBox<float>)
+				.SpinBoxStyle(&FCoreStyle::Get().GetWidgetStyle<FSpinBoxStyle>("NumericEntrySpinBox_Dark"))
+				.EditableTextBoxStyle(&FCoreStyle::Get().GetWidgetStyle<FEditableTextBoxStyle>("DarkEditableTextBox"))
+				.Font(IDetailLayoutBuilder::GetDetailFont())
+				.UndeterminedString(NSLOCTEXT("PropertyEditor", "MultipleValues", "Multiple Values"))
+				.Value(this, &SColorGradingColorWheel::GetComponentValue, ComponentIndex)
+				.OnValueChanged(this, &SColorGradingColorWheel::SetComponentValue, ComponentIndex)
+				.OnBeginSliderMovement(this, &SColorGradingColorWheel::BeginUsingComponentSlider, ComponentIndex)
+				.OnEndSliderMovement(this, &SColorGradingColorWheel::EndUsingComponentSlider, ComponentIndex)
+				.AllowSpin(ColorPropertyHandle.Pin()->GetNumOuterObjects() == 1)
+				.ShiftMultiplier(ColorPropertyMetadata->ShiftMultiplier)
+				.CtrlMultiplier(ColorPropertyMetadata->CtrlMultiplier)
+				.SupportDynamicSliderMinValue(this, &SColorGradingColorWheel::ComponentSupportsDynamicSliderValue, ColorPropertyMetadata->bSupportDynamicSliderMinValue, ComponentIndex)
+				.SupportDynamicSliderMaxValue(this, &SColorGradingColorWheel::ComponentSupportsDynamicSliderValue, ColorPropertyMetadata->bSupportDynamicSliderMaxValue, ComponentIndex)
+				.OnDynamicSliderMinValueChanged(this, &SColorGradingColorWheel::UpdateComponentDynamicSliderMinValue)
+				.OnDynamicSliderMaxValueChanged(this, &SColorGradingColorWheel::UpdateComponentDynamicSliderMaxValue)
+				.MinValue(ColorPropertyMetadata->MinValue)
+				.MaxValue(this, &SColorGradingColorWheel::GetComponentMaxValue, ColorPropertyMetadata->MaxValue, ComponentIndex)
+				.MinSliderValue(this, &SColorGradingColorWheel::GetComponentMinSliderValue, ColorPropertyMetadata->SliderMinValue, ComponentIndex)
+				.MaxSliderValue(this, &SColorGradingColorWheel::GetComponentMaxSliderValue, ColorPropertyMetadata->SliderMaxValue, ComponentIndex)
+				.SliderExponent(ColorPropertyMetadata->SliderExponent)
+				.SliderExponentNeutralValue(ColorPropertyMetadata->SliderMinValue.GetValue() + (ColorPropertyMetadata->SliderMaxValue.GetValue() - ColorPropertyMetadata->SliderMinValue.GetValue()) / 2.0f)
+				.Delta(this, &SColorGradingColorWheel::GetComponentSliderDeltaValue, ColorPropertyMetadata->Delta, ComponentIndex)
+				.ToolTipText(this, &SColorGradingColorWheel::GetComponentToolTipText, ComponentIndex)
+				.LabelPadding(FMargin(0))
+				.IsEnabled(this, &SColorGradingColorWheel::IsPropertyEnabled)
+				.Label()
+				[
+					LabelWidget
+				];
+
+			ColorSlidersVerticalBox->AddSlot()
+				.Padding(FMargin(0.0f, 2.0f, 3.0f, 0.0f))
+				.AutoHeight()
+				.HAlign(HAlign_Fill)
+				.VAlign(VAlign_Fill)
+				[
+					NumericEntryBox
+				];
+
+			ColorSlidersVerticalBox->AddSlot()
+				.Padding(FMargin(15.0f, 0.0f, 3.0f, 2.0f))
+				.AutoHeight()
+				.HAlign(HAlign_Fill)
+				.VAlign(VAlign_Bottom)
+				[
+					SNew(SBox)
+					.HeightOverride(6.0f)
+					[
+						SNew(SComplexGradient)
+						.GradientColors(TAttribute<TArray<FLinearColor>>::Create(TAttribute<TArray<FLinearColor>>::FGetter::CreateSP(this, &SColorGradingColorWheel::GetGradientColor, ComponentIndex)))
+						.Visibility(this, &SColorGradingColorWheel::GetGradientVisibility)
+					]
+				];
+		}
+
+		return ColorSlidersVerticalBox;
+	}
+
+	return SNullWidget::NullWidget;
+}
+
+SColorGradingColorWheel::FColorPropertyMetadata SColorGradingColorWheel::GetColorPropertyMetadata() const
+{
+	FColorPropertyMetadata PropertyMetadata;
+
+	if (ColorPropertyHandle.IsValid() && ColorPropertyHandle.Pin()->IsValidHandle())
+	{
+		FProperty* Property = ColorPropertyHandle.Pin()->GetProperty();
+		const FString ColorGradingModeString = Property->GetMetaData(TEXT("ColorGradingMode")).ToLower();
+
+		if (!ColorGradingModeString.IsEmpty())
+		{
+			if (ColorGradingModeString.Compare(TEXT("saturation")) == 0)
+			{
+				PropertyMetadata.ColorGradingMode = EColorGradingModes::Saturation;
+			}
+			else if (ColorGradingModeString.Compare(TEXT("contrast")) == 0)
+			{
+				PropertyMetadata.ColorGradingMode = EColorGradingModes::Contrast;
+			}
+			else if (ColorGradingModeString.Compare(TEXT("gamma")) == 0)
+			{
+				PropertyMetadata.ColorGradingMode = EColorGradingModes::Gamma;
+			}
+			else if (ColorGradingModeString.Compare(TEXT("gain")) == 0)
+			{
+				PropertyMetadata.ColorGradingMode = EColorGradingModes::Gain;
+			}
+			else if (ColorGradingModeString.Compare(TEXT("offset")) == 0)
+			{
+				PropertyMetadata.ColorGradingMode = EColorGradingModes::Offset;
+			}
+		}
+
+		const FString& MetaUIMinString = Property->GetMetaData(TEXT("UIMin"));
+		const FString& MetaUIMaxString = Property->GetMetaData(TEXT("UIMax"));
+		const FString& SliderExponentString = Property->GetMetaData(TEXT("SliderExponent"));
+		const FString& DeltaString = Property->GetMetaData(TEXT("Delta"));
+		const FString& LinearDeltaSensitivityString = Property->GetMetaData(TEXT("LinearDeltaSensitivity"));
+		const FString& ShiftMultiplierString = Property->GetMetaData(TEXT("ShiftMultiplier"));
+		const FString& CtrlMultiplierString = Property->GetMetaData(TEXT("CtrlMultiplier"));
+		const FString& SupportDynamicSliderMaxValueString = Property->GetMetaData(TEXT("SupportDynamicSliderMaxValue"));
+		const FString& SupportDynamicSliderMinValueString = Property->GetMetaData(TEXT("SupportDynamicSliderMinValue"));
+		const FString& ClampMinString = Property->GetMetaData(TEXT("ClampMin"));
+		const FString& ClampMaxString = Property->GetMetaData(TEXT("ClampMax"));
+
+		const FString& UIMinString = MetaUIMinString.Len() ? MetaUIMinString : ClampMinString;
+		const FString& UIMaxString = MetaUIMaxString.Len() ? MetaUIMaxString : ClampMaxString;
+
+		float ClampMin = TNumericLimits<float>::Lowest();
+		float ClampMax = TNumericLimits<float>::Max();
+
+		if (!ClampMinString.IsEmpty())
+		{
+			TTypeFromString<float>::FromString(ClampMin, *ClampMinString);
+		}
+
+		if (!ClampMaxString.IsEmpty())
+		{
+			TTypeFromString<float>::FromString(ClampMax, *ClampMaxString);
+		}
+
+		float UIMin = TNumericLimits<float>::Lowest();
+		float UIMax = TNumericLimits<float>::Max();
+		TTypeFromString<float>::FromString(UIMin, *UIMinString);
+		TTypeFromString<float>::FromString(UIMax, *UIMaxString);
+
+		const float ActualUIMin = FMath::Max(UIMin, ClampMin);
+		const float ActualUIMax = FMath::Min(UIMax, ClampMax);
+
+		PropertyMetadata.MinValue = ClampMinString.Len() ? ClampMin : TOptional<float>();
+		PropertyMetadata.MaxValue = ClampMaxString.Len() ? ClampMax : TOptional<float>();
+		PropertyMetadata.SliderMinValue = (UIMinString.Len()) ? ActualUIMin : TOptional<float>();
+		PropertyMetadata.SliderMaxValue = (UIMaxString.Len()) ? ActualUIMax : TOptional<float>();
+
+		if (SliderExponentString.Len())
+		{
+			TTypeFromString<float>::FromString(PropertyMetadata.SliderExponent, *SliderExponentString);
+		}
+
+		if (DeltaString.Len())
+		{
+			TTypeFromString<float>::FromString(PropertyMetadata.Delta, *DeltaString);
+		}
+
+		if (LinearDeltaSensitivityString.Len())
+		{
+			TTypeFromString<int32>::FromString(PropertyMetadata.LinearDeltaSensitivity, *LinearDeltaSensitivityString);
+			PropertyMetadata.Delta = (PropertyMetadata.Delta == 0.0f) ? 1.0f : PropertyMetadata.Delta;
+		}
+
+		PropertyMetadata.ShiftMultiplier = 10.f;
+		if (ShiftMultiplierString.Len())
+		{
+			TTypeFromString<float>::FromString(PropertyMetadata.ShiftMultiplier, *ShiftMultiplierString);
+		}
+
+		PropertyMetadata.CtrlMultiplier = 0.1f;
+		if (CtrlMultiplierString.Len())
+		{
+			TTypeFromString<float>::FromString(PropertyMetadata.CtrlMultiplier, *CtrlMultiplierString);
+		}
+
+		PropertyMetadata.bSupportDynamicSliderMaxValue = SupportDynamicSliderMaxValueString.Len() > 0 && SupportDynamicSliderMaxValueString.ToBool();
+		PropertyMetadata.bSupportDynamicSliderMinValue = SupportDynamicSliderMinValueString.Len() > 0 && SupportDynamicSliderMinValueString.ToBool();
+	}
+
+	return PropertyMetadata;
+}
+
+bool SColorGradingColorWheel::IsPropertyEnabled() const
+{
+	if (ColorPropertyHandle.IsValid() && ColorPropertyHandle.Pin()->IsValidHandle())
+	{
+		return ColorPropertyHandle.Pin()->IsEditable();
+	}
+
+	return false;
+}
+
+bool SColorGradingColorWheel::GetColor(FVector4& OutCurrentColor)
+{
+	if (ColorPropertyHandle.IsValid())
+	{
+		return ColorPropertyHandle.Pin()->GetValue(OutCurrentColor) == FPropertyAccess::Success;
+	}
+
+	return false;
+}
+
+void SColorGradingColorWheel::CommitColor(FVector4& NewValue, bool bShouldCommitValueChanges)
+{
+	FScopedTransaction Transaction(LOCTEXT("ColorWheel_TransactionName", "Color Grading Main Value"), bShouldCommitValueChanges);
+
+	if (ColorPropertyHandle.IsValid())
+	{
+		// Always perform a purely interactive change. We do this because it won't invoke reconstruction, which may cause that only the first 
+		// element gets updated due to its change causing a component reconstruction and the remaining vector element property handles updating 
+		// the trashed component.
+		ColorPropertyHandle.Pin()->SetValue(NewValue, EPropertyValueSetFlags::InteractiveChange | EPropertyValueSetFlags::NotTransactable);
+
+		// If not purely interactive, set the value with default flags.
+		if (bShouldCommitValueChanges || !bIsUsingColorPickerSlider)
+		{
+			ColorPropertyHandle.Pin()->SetValue(NewValue, EPropertyValueSetFlags::DefaultFlags);
+		}
+
+		TransactColorValue();
+	}
+}
+
+void SColorGradingColorWheel::TransactColorValue()
+{
+	if (ColorPropertyHandle.IsValid())
+	{
+		TArray<UObject*> OuterObjects;
+		ColorPropertyHandle.Pin()->GetOuterObjects(OuterObjects);
+		for (UObject* Object : OuterObjects)
+		{
+			if (!Object->HasAnyFlags(RF_Transactional))
+			{
+				Object->SetFlags(RF_Transactional);
+			}
+
+			SaveToTransactionBuffer(Object, false);
+			SnapshotTransactionBuffer(Object);
+		}
+	}
+}
+
+void SColorGradingColorWheel::BeginUsingColorPickerSlider()
+{
+	bIsUsingColorPickerSlider = true;
+	GEditor->BeginTransaction(LOCTEXT("ColorWheel_TransactionName", "Color Grading Main Value"));
+}
+
+void SColorGradingColorWheel::EndUsingColorPickerSlider()
+{
+	bIsUsingColorPickerSlider = false;
+	GEditor->EndTransaction();
+}
+
+void SColorGradingColorWheel::BeginUsingComponentSlider(uint32 ComponentIndex)
+{
+	bIsUsingComponentSlider = true;
+	GEditor->BeginTransaction(LOCTEXT("ColorWheel_TransactionName", "Color Grading Main Value"));
+}
+
+void SColorGradingColorWheel::EndUsingComponentSlider(float NewValue, uint32 ComponentIndex)
+{
+	bIsUsingComponentSlider = false;
+	SetComponentValue(NewValue, ComponentIndex);
+	GEditor->EndTransaction();
+}
+
+FText SColorGradingColorWheel::GetComponentLabelText(uint32 ComponentIndex) const
+{
+	static const FText RGBLabelText[] = {
+		LOCTEXT("ColorWheel_RedComponentLabel", "R"),
+		LOCTEXT("ColorWheel_GreenComponentLabel", "G"),
+		LOCTEXT("ColorWheel_BlueComponentLabel", "B")
+	};
+
+	static const FText HSVLabelText[] = {
+		LOCTEXT("ColorWheel_HueComponentLabel", "H"),
+		LOCTEXT("ColorWheel_SatComponentLabel", "S"),
+		LOCTEXT("ColorWheel_ValComponentLabel", "V")
+	};
+
+	if (ComponentIndex >= 0 && ComponentIndex < 3)
+	{ 
+		if (ColorDisplayMode.Get(EColorGradingColorDisplayMode::RGB) == EColorGradingColorDisplayMode::RGB)
+		{
+			return RGBLabelText[ComponentIndex];
+		}
+		else
+		{
+			return HSVLabelText[ComponentIndex];
+		}
+	}
+	else if (ComponentIndex == 3)
+	{
+		return LOCTEXT("ColorWheel_LuminanceComponentLabel", "Y");
+	}
+
+	return FText::GetEmpty();
+}
+
+TOptional<float> SColorGradingColorWheel::GetComponentValue(uint32 ComponentIndex) const
+{
+	if (ColorPropertyHandle.IsValid())
+	{
+		FVector4 ColorValue;
+		if (ColorPropertyHandle.Pin()->GetValue(ColorValue) == FPropertyAccess::Success)
+		{
+			float Value = 0.0f;
+
+			if (ColorDisplayMode.Get(EColorGradingColorDisplayMode::RGB) == EColorGradingColorDisplayMode::RGB)
+			{
+				Value = ColorValue[ComponentIndex];
+			}
+			else
+			{
+				FLinearColor HSVColor = FLinearColor(ColorValue.X, ColorValue.Y, ColorValue.Z, ColorValue.W).LinearRGBToHSV();
+				Value = HSVColor.Component(ComponentIndex);
+			}
+
+			return Value;
+		}
+	}
+
+	return TOptional<float>();
+}
+
+void SColorGradingColorWheel::SetComponentValue(float NewValue, uint32 ComponentIndex)
+{
+	if (ColorPropertyHandle.IsValid())
+	{
+		FVector4 CurrentColorValue;
+		if (ColorPropertyHandle.Pin()->GetValue(CurrentColorValue) == FPropertyAccess::Success)
+		{
+			FVector4 NewColorValue = CurrentColorValue;
+
+			if (ColorDisplayMode.Get(EColorGradingColorDisplayMode::RGB) == EColorGradingColorDisplayMode::RGB)
+			{
+				NewColorValue[ComponentIndex] = NewValue;
+			}
+			else
+			{
+				FLinearColor HSVColor = FLinearColor(CurrentColorValue.X, CurrentColorValue.Y, CurrentColorValue.Z, CurrentColorValue.W).LinearRGBToHSV();
+				HSVColor.Component(ComponentIndex) = NewValue;
+				NewColorValue = (FVector4)HSVColor.HSVToLinearRGB();
+			}
+
+			ColorPropertyHandle.Pin()->SetValue(NewColorValue, bIsUsingComponentSlider ? EPropertyValueSetFlags::InteractiveChange : EPropertyValueSetFlags::DefaultFlags);
+			TransactColorValue();
+		}
+	}
+}
+
+bool SColorGradingColorWheel::ComponentSupportsDynamicSliderValue(bool bDefaultValue, uint32 ComponentIndex) const
+{
+	if (bDefaultValue)
+	{
+		if (ColorDisplayMode.Get(EColorGradingColorDisplayMode::RGB) != EColorGradingColorDisplayMode::RGB)
+		{
+			return ComponentIndex >= 2;
+		}
+	}
+
+	return bDefaultValue;
+}
+
+void SColorGradingColorWheel::UpdateComponentDynamicSliderMinValue(float NewValue, TWeakPtr<SWidget> SourceWidget, bool bIsOriginator, bool bUpdateOnlyIfLower)
+{
+	if (!ComponentSliderDynamicMinValue.IsSet() || (NewValue < ComponentSliderDynamicMinValue.GetValue() && bUpdateOnlyIfLower) || !bUpdateOnlyIfLower)
+	{
+		ComponentSliderDynamicMinValue = NewValue;
+	}
+}
+
+void SColorGradingColorWheel::UpdateComponentDynamicSliderMaxValue(float NewValue, TWeakPtr<SWidget> SourceWidget, bool bIsOriginator, bool bUpdateOnlyIfHigher)
+{
+	if (!ComponentSliderDynamicMaxValue.IsSet() || (NewValue > ComponentSliderDynamicMaxValue.GetValue() && bUpdateOnlyIfHigher) || !bUpdateOnlyIfHigher)
+	{
+		ComponentSliderDynamicMaxValue = NewValue;
+	}
+}
+
+TOptional<float> SColorGradingColorWheel::GetComponentMaxValue(TOptional<float> DefaultValue, uint32 ComponentIndex) const
+{
+	if (ColorDisplayMode.Get(EColorGradingColorDisplayMode::RGB) != EColorGradingColorDisplayMode::RGB)
+	{
+		if (ComponentIndex == 0)
+		{
+			return 359.0f;
+		}
+		else if (ComponentIndex == 1)
+		{
+			return 1.0f;
+		}
+	}
+
+	return DefaultValue;
+}
+
+TOptional<float> SColorGradingColorWheel::GetComponentMinSliderValue(TOptional<float> DefaultValue, uint32 ComponentIndex) const
+{
+	if (ColorDisplayMode.Get(EColorGradingColorDisplayMode::RGB) != EColorGradingColorDisplayMode::RGB)
+	{
+		return 0.0f;
+	}
+
+	return ComponentSliderDynamicMinValue.IsSet() ? ComponentSliderDynamicMinValue.GetValue() : DefaultValue;
+}
+
+TOptional<float> SColorGradingColorWheel::GetComponentMaxSliderValue(TOptional<float> DefaultValue, uint32 ComponentIndex) const
+{
+	if (ColorDisplayMode.Get(EColorGradingColorDisplayMode::RGB) != EColorGradingColorDisplayMode::RGB)
+	{
+		if (ComponentIndex == 0)
+		{
+			return 359.0f;
+		}
+		else if (ComponentIndex == 1)
+		{
+			return 1.0f;
+		}
+	}
+
+	return ComponentSliderDynamicMaxValue.IsSet() ? ComponentSliderDynamicMaxValue : DefaultValue;
+}
+
+float SColorGradingColorWheel::GetComponentSliderDeltaValue(float DefaultValue, uint32 ComponentIndex) const
+{
+	if (ComponentIndex == 0 && ColorDisplayMode.Get(EColorGradingColorDisplayMode::RGB) != EColorGradingColorDisplayMode::RGB)
+	{
+		return 1.0f;
+	}
+
+	return DefaultValue;
+}
+
+FText SColorGradingColorWheel::GetComponentToolTipText(uint32 ComponentIndex) const
+{
+	static const FText RGBToolTipText[] =
+	{
+		LOCTEXT("ColorWheel_RedComponentToolTip", "Red"),
+		LOCTEXT("ColorWheel_GreenComponentToolTip", "Green"),
+		LOCTEXT("ColorWheel_BlueComponentToolTip", "Blue")
+	};
+
+	static const FText HSVToolTipText[] =
+	{
+		LOCTEXT("ColorWheel_HueComponentToolTip", "Hue"),
+		LOCTEXT("ColorWheel_SaturationComponentToolTip", "Saturation"),
+		LOCTEXT("ColorWheel_ValueComponentToolTip", "Value")
+	};
+
+	if (ComponentIndex >= 0 && ComponentIndex < 3)
+	{
+		if (ColorDisplayMode.Get(EColorGradingColorDisplayMode::RGB) != EColorGradingColorDisplayMode::RGB)
+		{
+			return HSVToolTipText[ComponentIndex];
+		}
+		else
+		{
+			return RGBToolTipText[ComponentIndex];
+		}
+	}
+	else if (ComponentIndex == 3)
+	{
+		return LOCTEXT("ColorWheel_LuminanceComponentToolTip", "Luminance");
+	}
+
+	return FText::GetEmpty();
+}
+
+TArray<FLinearColor> SColorGradingColorWheel::GetGradientColor(uint32 ComponentIndex)
+{
+	TArray<FLinearColor> GradientColors;
+
+	if (ColorPropertyHandle.IsValid())
+	{
+		const bool bIsRGB = ColorDisplayMode.Get(EColorGradingColorDisplayMode::RGB) == EColorGradingColorDisplayMode::RGB;
+		FVector4 ColorValue;
+			
+		if (ColorPropertyHandle.Pin()->GetValue(ColorValue) != FPropertyAccess::Success)
+		{
+			ColorValue = FVector4(0.0, 0.0, 0.0, 0.0);
+		}
+
+		if (bIsRGB || ComponentIndex > 0)
+		{
+			GradientColors.Add(GetGradientStartColor(ColorValue, bIsRGB, ComponentIndex));
+			GradientColors.Add(GetGradientEndColor(ColorValue, bIsRGB, ComponentIndex));
+			GradientColors.Add(GetGradientFillerColor(ColorValue, bIsRGB, ComponentIndex));
+		}
+		else // HSV Hue handling
+		{
+			for (int32 i = 0; i < 7; ++i)
+			{
+				GradientColors.Add(FLinearColor((i % 6) * 60.f, 1.f, 1.f).HSVToLinearRGB());
+			}
+		}
+	}
+
+	return GradientColors;
+}
+
+FLinearColor SColorGradingColorWheel::GetGradientStartColor(const FVector4& ColorValue, bool bIsRGB, uint32 ComponentIndex) const
+{
+	if (bIsRGB)
+	{
+		switch (ComponentIndex)
+		{
+		case 0:		return FLinearColor(0.0f, ColorValue.Y, ColorValue.Z, 1.0f);
+		case 1:		return FLinearColor(ColorValue.X, 0.0f, ColorValue.Z, 1.0f);
+		case 2:		return FLinearColor(ColorValue.X, ColorValue.Y, 0.0f, 1.0f);
+		case 3:		return FLinearColor(0.0f, 0.0f, 0.0f, 1.0f);
+		default:	return FLinearColor(ForceInit);
+		}
+	}
+	else
+	{
+		FLinearColor HSVColor = FLinearColor(ColorValue.X, ColorValue.Y, ColorValue.Z, ColorValue.W).LinearRGBToHSV();
+		switch (ComponentIndex)
+		{
+		case 0:		return FLinearColor(HSVColor.R, HSVColor.G, HSVColor.B, 1.0f);
+		case 1:		return FLinearColor(HSVColor.R, 0.0f, HSVColor.B, 1.0f).HSVToLinearRGB();
+		case 2:		return FLinearColor(HSVColor.R, HSVColor.G, 0.0f, 1.0f).HSVToLinearRGB();
+		case 3:		return FLinearColor(0.0f, 0.0f, 0.0f, 1.0f);
+		default:	return FLinearColor(ForceInit);
+		}
+	}
+}
+
+FLinearColor SColorGradingColorWheel::GetGradientEndColor(const FVector4& ColorValue, bool bIsRGB, uint32 ComponentIndex) const
+{
+	if (bIsRGB)
+	{
+		switch (ComponentIndex)
+		{
+		case 0:		return FLinearColor(1.0f, ColorValue.Y, ColorValue.Z, 1.0f);
+		case 1:		return FLinearColor(ColorValue.X, 1.0f, ColorValue.Z, 1.0f);
+		case 2:		return FLinearColor(ColorValue.X, ColorValue.Y, 1.0f, 1.0f);
+		case 3:		return FLinearColor(ColorValue.X, ColorValue.Y, ColorValue.Z, 1.0f);
+		default:	return FLinearColor(ForceInit);
+		}
+	}
+	else
+	{
+		FLinearColor HSVColor = FLinearColor(ColorValue.X, ColorValue.Y, ColorValue.Z, ColorValue.W).LinearRGBToHSV();
+		switch (ComponentIndex)
+		{
+		case 0:		return FLinearColor(HSVColor.R, HSVColor.G, HSVColor.B, 1.0f);
+		case 1:		return FLinearColor(HSVColor.R, 1.0f, HSVColor.B, 1.0f).HSVToLinearRGB();
+		case 2:		return FLinearColor(HSVColor.R, HSVColor.G, 1.0f, 1.0f).HSVToLinearRGB();
+		case 3:		return FLinearColor(HSVColor.R, HSVColor.G, HSVColor.B, 1.0f).HSVToLinearRGB();
+		default:	return FLinearColor(ForceInit);
+		}
+	}
+}
+
+FLinearColor SColorGradingColorWheel::GetGradientFillerColor(const FVector4& ColorValue, bool bIsRGB, uint32 ComponentIndex) const
+{
+	const float MaxValue = 1.0f/*SpinBoxMinMaxSliderValues.CurrentMaxSliderValue.GetValue()*/;
+	if (bIsRGB)
+	{
+		switch (ComponentIndex)
+		{
+		case 0:		return FLinearColor(MaxValue, ColorValue.Y, ColorValue.Z, 1.0f);
+		case 1:		return FLinearColor(ColorValue.X, MaxValue, ColorValue.Z, 1.0f);
+		case 2:		return FLinearColor(ColorValue.X, ColorValue.Y, MaxValue, 1.0f);
+		case 3:		return FLinearColor(ColorValue.X, ColorValue.Y, ColorValue.Z, 1.0f);
+		default:	return FLinearColor(ForceInit);
+		}
+	}
+	else
+	{
+		FLinearColor HSVColor = FLinearColor(ColorValue.X, ColorValue.Y, ColorValue.Z, ColorValue.W).LinearRGBToHSV();
+		switch (ComponentIndex)
+		{
+		case 0:		return FLinearColor(HSVColor.R, HSVColor.G, HSVColor.B, 1.0f);
+		case 1:		return FLinearColor(HSVColor.R, 1.0f, HSVColor.B, 1.0f).HSVToLinearRGB();
+		case 2:		return FLinearColor(HSVColor.R, HSVColor.G, MaxValue, 1.0f).HSVToLinearRGB();
+		case 3:		return FLinearColor(HSVColor.R, HSVColor.G, HSVColor.B, 1.0f).HSVToLinearRGB();
+		default:	return FLinearColor(ForceInit);
+		}
+	}
+}
+
+EVisibility SColorGradingColorWheel::GetGradientVisibility() const
+{
+	bool bShowGradient =
+		ColorPropertyMetadata->ColorGradingMode != EColorGradingModes::Offset &&
+		ColorPropertyMetadata->ColorGradingMode != EColorGradingModes::Invalid;
+
+	return bShowGradient ? EVisibility::Visible : EVisibility::Hidden;
+}
+
+TOptional<float> SColorGradingColorWheel::GetMetadataMinValue() const
+{
+	if (ColorPropertyMetadata.IsSet())
+	{
+		return ColorPropertyMetadata->MinValue;
+	}
+
+	return TOptional<float>();
+}
+
+TOptional<float> SColorGradingColorWheel::GetMetadataMaxValue() const
+{
+	if (ColorPropertyMetadata.IsSet())
+	{
+		return ColorPropertyMetadata->MaxValue;
+	}
+
+	return TOptional<float>();
+}
+
+TOptional<float> SColorGradingColorWheel::GetMetadataSliderMinValue() const
+{
+	if (ColorPropertyMetadata.IsSet())
+	{
+		return ColorPropertyMetadata->SliderMinValue;
+	}
+
+	return TOptional<float>();
+}
+
+TOptional<float> SColorGradingColorWheel::GetMetadataSliderMaxValue() const
+{
+	if (ColorPropertyMetadata.IsSet())
+	{
+		return ColorPropertyMetadata->SliderMaxValue;
+	}
+
+	return TOptional<float>();
+}
+
+float SColorGradingColorWheel::GetMetadataDelta() const
+{
+	if (ColorPropertyMetadata.IsSet())
+	{
+		return ColorPropertyMetadata->Delta;
+	}
+
+	return 0.0f;
+}
+
+float SColorGradingColorWheel::GetMetadataShiftMultiplier() const
+{
+	if (ColorPropertyMetadata.IsSet())
+	{
+		return ColorPropertyMetadata->ShiftMultiplier;
+	}
+
+	return 10.f;
+}
+
+float SColorGradingColorWheel::GetMetadataCtrlMultiplier() const
+{
+	if (ColorPropertyMetadata.IsSet())
+	{
+		return ColorPropertyMetadata->CtrlMultiplier;
+	}
+
+	return 0.1f;
+}
+
+bool SColorGradingColorWheel::GetMetadataSupportDynamicSliderMinValue() const
+{
+	if (ColorPropertyMetadata.IsSet())
+	{
+		return ColorPropertyMetadata->bSupportDynamicSliderMinValue;
+	}
+
+	return false;
+}
+
+bool SColorGradingColorWheel::GetMetadataSupportDynamicSliderMaxValue() const
+{
+	if (ColorPropertyMetadata.IsSet())
+	{
+		return ColorPropertyMetadata->bSupportDynamicSliderMaxValue;
+	}
+
+	return false;
+}
+
+#undef LOCTEXT_NAMESPACE
\ No newline at end of file
diff --git a/Engine/Plugins/Editor/ColorGrading/Source/ColorGradingEditor/Private/SColorGradingColorWheelPanel.cpp b/Engine/Plugins/Editor/ColorGrading/Source/ColorGradingEditor/Private/SColorGradingColorWheelPanel.cpp
new file mode 100644
--- /dev/null
+++ b/Engine/Plugins/Editor/ColorGrading/Source/ColorGradingEditor/Private/SColorGradingColorWheelPanel.cpp
@@ -0,0 +1,727 @@
+// Copyright Epic Games, Inc. All Rights Reserved.
+
+#include "SColorGradingColorWheelPanel.h"
+
+#include "ColorGradingCommands.h"
+#include "ColorGradingPanelState.h"
+#include "SColorGradingColorWheel.h"
+#include "DetailView/SColorGradingDetailView.h"
+
+#include "Framework/MultiBox/MultiBoxBuilder.h"
+#include "IDetailTreeNode.h"
+#include "IPropertyRowGenerator.h"
+#include "Modules/ModuleManager.h"
+#include "PropertyHandle.h"
+#include "Widgets/SBoxPanel.h"
+#include "Widgets/Input/SComboButton.h"
+#include "Widgets/Layout/SBox.h"
+#include "Widgets/Layout/SSpacer.h"
+#include "Widgets/Layout/SSplitter.h"
+ 
+#define LOCTEXT_NAMESPACE "ColorGradingEditor"
+
+SColorGradingColorWheelPanel::~SColorGradingColorWheelPanel()
+{
+	if (ColorGradingDataModel)
+	{
+		ColorGradingDataModel->OnColorGradingGroupSelectionChanged().RemoveAll(this);
+		ColorGradingDataModel->OnColorGradingElementSelectionChanged().RemoveAll(this);
+	}
+}
+
+void SColorGradingColorWheelPanel::Construct(const FArguments& InArgs)
+{
+	ColorGradingDataModel = InArgs._ColorGradingDataModelSource;
+
+	if (ColorGradingDataModel)
+	{
+		ColorGradingDataModel->OnColorGradingGroupSelectionChanged().AddSP(this, &SColorGradingColorWheelPanel::OnColorGradingGroupSelectionChanged);
+		ColorGradingDataModel->OnColorGradingElementSelectionChanged().AddSP(this, &SColorGradingColorWheelPanel::OnColorGradingElementSelectionChanged);
+	}
+
+	CommandList = MakeShared<FUICommandList>();
+	BindCommands();
+
+	ColorWheelOrientation = EOrientation::Orient_Vertical;
+
+	ColorWheels.AddDefaulted(NumColorWheels);
+	HiddenColorWheels.AddZeroed(NumColorWheels);
+
+	for (int32 Index = 0; Index < NumColorWheels; ++Index)
+    ChildSlot
+	[
+		SNew(SSplitter)
+		.Orientation(Orient_Horizontal)
+		.PhysicalSplitterHandleSize(1.0f)
+		.HitDetectionSplitterHandleSize(5.0f)
+		.Style(FAppStyle::Get(), "DetailsView.Splitter")
+
+		+ SSplitter::Slot()
+		.Value(0.8f)
+		[
+			SNew(SVerticalBox)
+			.Visibility(this, &SColorGradingColorWheelPanel::GetColorWheelPanelVisibility)
+
+			// Toolbar slot
+			+ SVerticalBox::Slot()
+			.AutoHeight()
+			.Padding(6, 4)
+			[
+				SNew(SHorizontalBox)
+
+				+ SHorizontalBox::Slot()
+				.AutoWidth()
+				[
+					SAssignNew(ColorGradingGroupPropertyBox, SBox)
+				]
+
+				+ SHorizontalBox::Slot()
+				.FillWidth(1.0f)
+				[
+					SNew(SSpacer)
+				]
+
+				+ SHorizontalBox::Slot()
+				.AutoWidth()
+				[
+					SAssignNew(ColorGradingElementsToolBarBox, SHorizontalBox)
+				]
+
+				+ SHorizontalBox::Slot()
+				.FillWidth(1.0f)
+				[
+					SNew(SSpacer)
+				]
+
+				+ SHorizontalBox::Slot()
+				.AutoWidth()
+				[
+					MakeColorDisplayModeCheckbox()
+				]
+
+				+ SHorizontalBox::Slot()
+				.AutoWidth()
+				.Padding(5.0f, 0.0f, 0.0f, 0.0f)
+				.HAlign(HAlign_Right)
+				.VAlign(VAlign_Center)
+				[
+					SNew(SComboButton)
+					.ComboButtonStyle(&FAppStyle::Get().GetWidgetStyle<FComboButtonStyle>("SimpleComboButton"))
+					.OnGetMenuContent(this, &SColorGradingColorWheelPanel::MakeSettingsMenu)
+					.HasDownArrow(false)
+					.ContentPadding(FMargin(1.0f, 0.0f))
+					.ButtonContent()
+					[
+						SNew(SImage)
+						.ColorAndOpacity(FSlateColor::UseForeground())
+						.Image(FAppStyle::Get().GetBrush("Icons.Settings"))
+					]
+				]
+			]
+
+			+ SVerticalBox::Slot()
+			.FillHeight(1.0f)
+			.Padding(6, 4)
+			[
+				SNew(SHorizontalBox)
+
+				+ SHorizontalBox::Slot()
+				.FillWidth(1.0f)
+				.HAlign(HAlign_Fill)
+				.Padding(2, 0)
+				[
+					SAssignNew(ColorWheels[0], SColorGradingColorWheel)
+					.ColorDisplayMode(this, &SColorGradingColorWheelPanel::GetColorDisplayMode)
+					.Orientation(ColorWheelOrientation)
+					.Visibility(this, &SColorGradingColorWheelPanel::GetColorWheelVisibility, 0)
+				]
+
+				+ SHorizontalBox::Slot()
+				.FillWidth(1.0f)
+				.HAlign(HAlign_Fill)
+				.Padding(2, 0)
+				[
+					SAssignNew(ColorWheels[1], SColorGradingColorWheel)
+					.ColorDisplayMode(this, &SColorGradingColorWheelPanel::GetColorDisplayMode)
+					.Orientation(ColorWheelOrientation)
+					.Visibility(this, &SColorGradingColorWheelPanel::GetColorWheelVisibility, 1)
+				]
+
+				+ SHorizontalBox::Slot()
+				.FillWidth(1.0f)
+				.HAlign(HAlign_Fill)
+				.Padding(2, 0)
+				[
+					SAssignNew(ColorWheels[2], SColorGradingColorWheel)
+					.ColorDisplayMode(this, &SColorGradingColorWheelPanel::GetColorDisplayMode)
+					.Orientation(ColorWheelOrientation)
+					.Visibility(this, &SColorGradingColorWheelPanel::GetColorWheelVisibility, 2)
+				]
+
+				+ SHorizontalBox::Slot()
+				.FillWidth(1.0f)
+				.HAlign(HAlign_Fill)
+				.Padding(2, 0)
+				[
+					SAssignNew(ColorWheels[3], SColorGradingColorWheel)
+					.ColorDisplayMode(this, &SColorGradingColorWheelPanel::GetColorDisplayMode)
+					.Orientation(ColorWheelOrientation)
+					.Visibility(this, &SColorGradingColorWheelPanel::GetColorWheelVisibility, 3)
+				]
+
+				+ SHorizontalBox::Slot()
+				.FillWidth(1.0f)
+				.HAlign(HAlign_Fill)
+				.Padding(2, 0)
+				[
+					SAssignNew(ColorWheels[4], SColorGradingColorWheel)
+					.ColorDisplayMode(this, &SColorGradingColorWheelPanel::GetColorDisplayMode)
+					.Orientation(ColorWheelOrientation)
+					.Visibility(this, &SColorGradingColorWheelPanel::GetColorWheelVisibility, 4)
+				]
+			]
+		]
+
+		+ SSplitter::Slot()
+		.Value(0.2f)
+		[
+			SAssignNew(DetailView, SColorGradingDetailView)
+			.PropertyRowGeneratorSource(ColorGradingDataModel->GetPropertyRowGenerator())
+			.OnFilterDetailTreeNode(this, &SColorGradingColorWheelPanel::FilterDetailTreeNode)
+		]
+	];
+}
+
+void SColorGradingColorWheelPanel::Refresh()
+{
+	if (ColorGradingDataModel)
+	{
+		if (FColorGradingEditorDataModel::FColorGradingGroup* ColorGradingGroup = ColorGradingDataModel->GetSelectedColorGradingGroup())
+		{
+			FillColorGradingGroupProperty(*ColorGradingGroup);
+			FillColorGradingElementsToolBar(ColorGradingGroup->ColorGradingElements);
+
+			if (const FColorGradingEditorDataModel::FColorGradingElement* ColorGradingElement = ColorGradingDataModel->GetSelectedColorGradingElement())
+			{
+				FillColorWheels(*ColorGradingElement);
+			}
+			else
+			{
+				ClearColorWheels();
+			}
+		}
+		else
+		{
+			ClearColorGradingGroupProperty();
+			ClearColorGradingElementsToolBar();
+			ClearColorWheels();
+		}
+
+		DetailView->Refresh();
+	}
+}
+
+void SColorGradingColorWheelPanel::GetPanelState(FColorGradingPanelState& OutPanelState)
+{
+	OutPanelState.HiddenColorWheels = HiddenColorWheels;
+	OutPanelState.ColorDisplayMode = ColorDisplayMode;
+	OutPanelState.ColorWheelOrientation = ColorWheelOrientation;
+}
+
+void SColorGradingColorWheelPanel::SetPanelState(const FColorGradingPanelState& InPanelState)
+{
+	// TODO: These could also be output to a config file to be stored between runs
+	HiddenColorWheels = InPanelState.HiddenColorWheels;
+	ColorDisplayMode = InPanelState.ColorDisplayMode;
+	SetColorWheelOrientation(InPanelState.ColorWheelOrientation);
+}
+
+void SColorGradingColorWheelPanel::BindCommands()
+{
+	const FColorGradingCommands& Commands = FColorGradingCommands::Get();
+
+	CommandList->MapAction(
+		Commands.SaturationColorWheelVisibility,
+		FExecuteAction::CreateSP(this, &SColorGradingColorWheelPanel::ToggleColorWheelVisible, 3),
+		FCanExecuteAction(),
+		FIsActionChecked::CreateSP(this, &SColorGradingColorWheelPanel::IsColorWheelVisible, 3)
+	);
+
+	CommandList->MapAction(
+		Commands.ContrastColorWheelVisibility,
+		FExecuteAction::CreateSP(this, &SColorGradingColorWheelPanel::ToggleColorWheelVisible, 4),
+		FCanExecuteAction(),
+		FIsActionChecked::CreateSP(this, &SColorGradingColorWheelPanel::IsColorWheelVisible, 4)
+	);
+
+	CommandList->MapAction(
+		Commands.ColorWheelSliderOrientationHorizontal,
+		FExecuteAction::CreateSP(this, &SColorGradingColorWheelPanel::SetColorWheelOrientation, EOrientation::Orient_Horizontal),
+		FCanExecuteAction(),
+		FIsActionChecked::CreateSP(this, &SColorGradingColorWheelPanel::IsColorWheelOrientationSelected, EOrientation::Orient_Horizontal)
+	);
+
+	CommandList->MapAction(
+		Commands.ColorWheelSliderOrientationVertical,
+		FExecuteAction::CreateSP(this, &SColorGradingColorWheelPanel::SetColorWheelOrientation, EOrientation::Orient_Vertical),
+		FCanExecuteAction(),
+		FIsActionChecked::CreateSP(this, &SColorGradingColorWheelPanel::IsColorWheelOrientationSelected, EOrientation::Orient_Vertical)
+	);
+}
+
+TSharedRef<SWidget> SColorGradingColorWheelPanel::MakeColorDisplayModeCheckbox()
+{
+	return SNew(SHorizontalBox)
+
+		+ SHorizontalBox::Slot()
+		.AutoWidth()
+		.HAlign(HAlign_Left)
+		.VAlign(VAlign_Top)
+		.Padding(FMargin(0.0f, 0.0f, 3.0f, 0.0f))
+		[
+			SNew(SCheckBox)
+			.Style(FAppStyle::Get(), "ToggleButtonCheckbox")
+			.Type(ESlateCheckBoxType::ToggleButton)
+			.IsChecked(this, &SColorGradingColorWheelPanel::IsColorDisplayModeChecked, EColorGradingColorDisplayMode::RGB)
+			.OnCheckStateChanged(this, &SColorGradingColorWheelPanel::OnColorDisplayModeCheckedChanged, EColorGradingColorDisplayMode::RGB)
+			.ToolTipText(this, &SColorGradingColorWheelPanel::GetColorDisplayModeToolTip, EColorGradingColorDisplayMode::RGB)
+			.Padding(4)
+			.Content()
+			[
+				SNew(STextBlock)
+				.Text(this, &SColorGradingColorWheelPanel::GetColorDisplayModeLabel, EColorGradingColorDisplayMode::RGB)
+				.Font(FAppStyle::GetFontStyle("StandardDialog.SmallFont"))
+			]
+		]
+		+ SHorizontalBox::Slot()
+		.AutoWidth()
+		.HAlign(HAlign_Left)			
+		.VAlign(VAlign_Top)
+		.Padding(FMargin(0.0f, 0.0f, 3.0f, 0.0f))
+		[
+			SNew(SCheckBox)
+			.Style(FAppStyle::Get(), "ToggleButtonCheckbox")
+			.Type(ESlateCheckBoxType::ToggleButton)
+			.IsChecked(this, &SColorGradingColorWheelPanel::IsColorDisplayModeChecked, EColorGradingColorDisplayMode::HSV)
+			.OnCheckStateChanged(this, &SColorGradingColorWheelPanel::OnColorDisplayModeCheckedChanged, EColorGradingColorDisplayMode::HSV)
+			.ToolTipText(this, &SColorGradingColorWheelPanel::GetColorDisplayModeToolTip, EColorGradingColorDisplayMode::HSV)
+			.Padding(4)
+			.Content()
+			[
+				SNew(STextBlock)
+				.Text(this, &SColorGradingColorWheelPanel::GetColorDisplayModeLabel, EColorGradingColorDisplayMode::HSV)
+				.Font(FAppStyle::GetFontStyle("StandardDialog.SmallFont"))
+			]
+		];
+}
+
+TSharedRef<SWidget> SColorGradingColorWheelPanel::MakeSettingsMenu()
+{
+	const bool bShouldCloseWindowAfterMenuSelection = true;
+	FMenuBuilder MenuBuilder(bShouldCloseWindowAfterMenuSelection, CommandList);
+
+	const FColorGradingCommands& Commands = FColorGradingCommands::Get();
+
+	MenuBuilder.BeginSection(TEXT("ColorWheelVisibility"), LOCTEXT("ColorWheelPanel_ShowLabel", "Show"));
+	{
+		MenuBuilder.AddMenuEntry(Commands.SaturationColorWheelVisibility);
+		MenuBuilder.AddMenuEntry(Commands.ContrastColorWheelVisibility);
+	}
+	MenuBuilder.EndSection();
+
+	MenuBuilder.BeginSection(TEXT("ColorWheelSliders"), LOCTEXT("ColorWheelPanel_SlidersLabel", "Sliders"));
+	{
+		MenuBuilder.AddMenuEntry(Commands.ColorWheelSliderOrientationVertical);
+		MenuBuilder.AddMenuEntry(Commands.ColorWheelSliderOrientationHorizontal);
+	}
+	MenuBuilder.EndSection();
+
+	return MenuBuilder.MakeWidget();
+}
+
+void SColorGradingColorWheelPanel::FillColorGradingGroupProperty(const FColorGradingEditorDataModel::FColorGradingGroup& ColorGradingGroup)
+{
+	if (ColorGradingGroupPropertyBox.IsValid())
+	{
+		TSharedRef<SHorizontalBox> PropertyNameBox = SNew(SHorizontalBox);
+
+		if (ColorGradingGroup.EditConditionPropertyHandle.IsValid())
+		{
+			if (TSharedPtr<IDetailTreeNode> EditConditionTreeNode = ColorGradingDataModel->GetPropertyRowGenerator()->FindTreeNode(ColorGradingGroup.EditConditionPropertyHandle))
+			{
+				FNodeWidgets EditConditionWidgets = EditConditionTreeNode->CreateNodeWidgets();
+
+				if (EditConditionWidgets.ValueWidget.IsValid())
+				{
+					PropertyNameBox->AddSlot()
+						.HAlign(HAlign_Left)
+						.VAlign(VAlign_Center)
+						.Padding(2, 0, 4, 0)
+						.AutoWidth()
+						[
+							EditConditionWidgets.ValueWidget.ToSharedRef()
+						];
+				}
+			}
+		}
+
+		TSharedRef<SWidget> GroupHeaderWidget = ColorGradingGroup.GroupHeaderWidget.IsValid()
+			? ColorGradingGroup.GroupHeaderWidget.ToSharedRef()
+			: SNew(STextBlock).Text(ColorGradingGroup.DisplayName).Font(FAppStyle::Get().GetFontStyle("NormalFontBold"));
+
+		PropertyNameBox->AddSlot()
+			.HAlign(HAlign_Left)
+			.VAlign(VAlign_Center)
+			.Padding(2, 0, 0, 0)
+			[
+				GroupHeaderWidget
+			];
+
+		ColorGradingGroupPropertyBox->SetContent(PropertyNameBox);
+	}
+}
+
+void SColorGradingColorWheelPanel::ClearColorGradingGroupProperty()
+{
+	ColorGradingGroupPropertyBox->SetContent(SNullWidget::NullWidget);
+}
+
+void SColorGradingColorWheelPanel::FillColorGradingElementsToolBar(const TArray<FColorGradingEditorDataModel::FColorGradingElement>& ColorGradingElements)
+{
+	ColorGradingElementsToolBarBox->ClearChildren();
+
+	for (const FColorGradingEditorDataModel::FColorGradingElement& Element : ColorGradingElements)
+	{
+		ColorGradingElementsToolBarBox->AddSlot()
+			.AutoWidth()
+			.Padding(2.0f, 0.0f, 2.0f, 0.0f)
+			[
+				SNew(SCheckBox)
+				.Style(FAppStyle::Get(), "DetailsView.SectionButton")
+				.OnCheckStateChanged(this, &SColorGradingColorWheelPanel::OnColorGradingElementCheckedChanged, Element.DisplayName)
+				.IsChecked(this, &SColorGradingColorWheelPanel::IsColorGradingElementSelected, Element.DisplayName)
+				[
+					SNew(STextBlock)
+					.TextStyle(FAppStyle::Get(), "SmallText")
+					.Text(Element.DisplayName)
+				]
+			];
+	}
+}
+
+void SColorGradingColorWheelPanel::ClearColorGradingElementsToolBar()
+{
+	ColorGradingElementsToolBarBox->ClearChildren();
+}
+
+void SColorGradingColorWheelPanel::FillColorWheels(const FColorGradingEditorDataModel::FColorGradingElement& ColorGradingElement)
+{
+	auto FillColorWheel = [this](const TSharedPtr<SColorGradingColorWheel>& ColorWheel, const TSharedPtr<IPropertyHandle>& PropertyHandle)
+	{
+		if (ColorWheel.IsValid())
+		{
+			ColorWheel->SetColorPropertyHandle(PropertyHandle);
+			ColorWheel->SetHeaderContent(CreateColorWheelHeaderWidget(PropertyHandle));
+		}
+	};
+
+	FillColorWheel(ColorWheels[0], ColorGradingElement.OffsetPropertyHandle);
+	FillColorWheel(ColorWheels[1], ColorGradingElement.GammaPropertyHandle);
+	FillColorWheel(ColorWheels[2], ColorGradingElement.GainPropertyHandle);
+	FillColorWheel(ColorWheels[3], ColorGradingElement.SaturationPropertyHandle);
+	FillColorWheel(ColorWheels[4], ColorGradingElement.ContrastPropertyHandle);
+}
+
+void SColorGradingColorWheelPanel::ClearColorWheels()
+{
+	for (const TSharedPtr<SColorGradingColorWheel>& ColorWheel : ColorWheels)
+	{
+		if (ColorWheel.IsValid())
+		{
+			ColorWheel->SetColorPropertyHandle(nullptr);
+			ColorWheel->SetHeaderContent(SNullWidget::NullWidget);
+		}
+	};
+}
+
+TSharedRef<SWidget> SColorGradingColorWheelPanel::CreateColorWheelHeaderWidget(const TSharedPtr<IPropertyHandle>& ColorPropertyHandle)
+{
+	if (TSharedPtr<IDetailTreeNode> TreeNode = ColorGradingDataModel->GetPropertyRowGenerator()->FindTreeNode(ColorPropertyHandle))
+	{
+		FNodeWidgets NodeWidgets = TreeNode->CreateNodeWidgets();
+
+		TSharedRef<SHorizontalBox> PropertyNameBox = SNew(SHorizontalBox);
+
+		PropertyNameBox->AddSlot()
+			.FillWidth(1.0f)
+			[
+				SNew(SSpacer)
+			];
+
+		if (NodeWidgets.EditConditionWidget.IsValid())
+		{
+			PropertyNameBox->AddSlot()
+				.VAlign(VAlign_Center)
+				.HAlign(HAlign_Left)
+				.Padding(2, 0, 0, 0)
+				.AutoWidth()
+				[
+					NodeWidgets.EditConditionWidget.ToSharedRef()
+				];
+		}
+
+		if (NodeWidgets.NameWidget.IsValid())
+		{
+			PropertyNameBox->AddSlot()
+				.HAlign(NodeWidgets.NameWidgetLayoutData.HorizontalAlignment)
+				.VAlign(NodeWidgets.NameWidgetLayoutData.VerticalAlignment)
+				.Padding(2, 0, 0, 0)
+				[
+					NodeWidgets.NameWidget.ToSharedRef()
+				];
+
+			PropertyNameBox->AddSlot()
+				.HAlign(HAlign_Right)
+				.VAlign(VAlign_Center)
+				.FillWidth(1.0f)
+				[
+					SNew(SBox)
+					.WidthOverride(22.0f)
+					[
+						CreateColorPropertyExtensions(ColorPropertyHandle, TreeNode)
+					]
+				];
+		}
+
+		return PropertyNameBox;
+	}
+
+	return SNullWidget::NullWidget;
+}
+
+TSharedRef<SWidget> SColorGradingColorWheelPanel::CreateColorPropertyExtensions(const TSharedPtr<IPropertyHandle>& ColorPropertyHandle, const TSharedPtr<IDetailTreeNode>& DetailTreeNode)
+{
+	TArray<FPropertyRowExtensionButton> ExtensionButtons;
+
+	// Use a weak pointer to pass into delegates
+	TWeakPtr<IPropertyHandle> WeakColorPropertyHandle = ColorPropertyHandle;
+
+	FPropertyRowExtensionButton& ResetToDefaultButton = ExtensionButtons.AddDefaulted_GetRef();
+	ResetToDefaultButton.Label = NSLOCTEXT("PropertyEditor", "ResetToDefault", "Reset to Default");
+	ResetToDefaultButton.UIAction = FUIAction(
+		FExecuteAction::CreateLambda([WeakColorPropertyHandle]
+		{
+			if (WeakColorPropertyHandle.IsValid())
+			{
+				WeakColorPropertyHandle.Pin()->ResetToDefault();
+			}
+		}),
+		FCanExecuteAction::CreateLambda([WeakColorPropertyHandle]
+		{
+			const bool bIsEditable = WeakColorPropertyHandle.Pin()->IsEditable();
+			return bIsEditable;
+		}),
+		FIsActionChecked(),
+		FIsActionButtonVisible::CreateLambda([WeakColorPropertyHandle]
+		{
+			bool bShowResetToDefaultButton = false;
+			if (WeakColorPropertyHandle.IsValid())
+			{
+				if (!WeakColorPropertyHandle.Pin()->HasMetaData("NoResetToDefault") && !WeakColorPropertyHandle.Pin()->GetInstanceMetaData("NoResetToDefault"))
+				{
+					bShowResetToDefaultButton = WeakColorPropertyHandle.Pin()->CanResetToDefault();
+				}
+			}
+
+			return bShowResetToDefaultButton;
+		})
+	);
+
+	ResetToDefaultButton.Icon = FSlateIcon(FAppStyle::Get().GetStyleSetName(), "PropertyWindow.DiffersFromDefault");
+	ResetToDefaultButton.ToolTip = NSLOCTEXT("PropertyEditor", "ResetToDefaultPropertyValueToolTip", "Reset this property to its default value.");
+
+	// Add any global row extensions that are registered for the color property
+	FPropertyEditorModule& PropertyEditorModule = FModuleManager::Get().GetModuleChecked<FPropertyEditorModule>("PropertyEditor");
+
+	FOnGenerateGlobalRowExtensionArgs Args;
+	Args.OwnerTreeNode = DetailTreeNode;
+	Args.PropertyHandle = ColorPropertyHandle;
+
+	PropertyEditorModule.GetGlobalRowExtensionDelegate().Broadcast(Args, ExtensionButtons);
+
+	FSlimHorizontalToolBarBuilder ToolbarBuilder(TSharedPtr<FUICommandList>(), FMultiBoxCustomization::None);
+	ToolbarBuilder.SetLabelVisibility(EVisibility::Collapsed);
+	ToolbarBuilder.SetStyle(&FAppStyle::Get(), "DetailsView.ExtensionToolBar");
+	ToolbarBuilder.SetIsFocusable(false);
+
+	for (const FPropertyRowExtensionButton& Extension : ExtensionButtons)
+	{
+		ToolbarBuilder.AddToolBarButton(Extension.UIAction, NAME_None, Extension.Label, Extension.ToolTip, Extension.Icon);
+	}
+
+	return ToolbarBuilder.MakeWidget();
+}
+
+bool SColorGradingColorWheelPanel::FilterDetailTreeNode(const TSharedRef<IDetailTreeNode>& InDetailTreeNode)
+{
+	if (ColorGradingDataModel)
+	{
+		if (FColorGradingEditorDataModel::FColorGradingGroup* ColorGradingGroup = ColorGradingDataModel->GetSelectedColorGradingGroup())
+		{
+			// Filter out any categories that are not configured by the data model to be displayed in the details section or subsection.
+			// All other nodes (which will be any child of the category), should be displayed.
+			if (InDetailTreeNode->GetNodeType() == EDetailNodeType::Category)
+			{
+				return ColorGradingGroup->DetailsViewCategories.Contains(InDetailTreeNode->GetNodeName());
+			}
+			else
+			{
+				return true;
+			}
+		}
+	}
+
+	return false;
+}
+
+void SColorGradingColorWheelPanel::SetColorWheelOrientation(EOrientation NewOrientation)
+{
+	if (ColorWheelOrientation != NewOrientation)
+	{
+		ColorWheelOrientation = NewOrientation;
+
+		for (const TSharedPtr<SColorGradingColorWheel>& ColorWheel : ColorWheels)
+		{
+			if (ColorWheel.IsValid())
+			{
+				ColorWheel->SetOrientation(ColorWheelOrientation);
+			}
+		};
+	}
+}
+
+bool SColorGradingColorWheelPanel::IsColorWheelOrientationSelected(EOrientation Orientation) const
+{
+	return ColorWheelOrientation == Orientation;
+}
+
+void SColorGradingColorWheelPanel::ToggleColorWheelVisible(int32 ColorWheelIndex)
+{
+	if (ColorWheelIndex >= 0 && ColorWheelIndex < HiddenColorWheels.Num())
+	{
+		HiddenColorWheels[ColorWheelIndex] = !HiddenColorWheels[ColorWheelIndex];
+	}
+}
+
+bool SColorGradingColorWheelPanel::IsColorWheelVisible(int32 ColorWheelIndex)
+{
+	if (ColorWheelIndex >= 0 && ColorWheelIndex < HiddenColorWheels.Num())
+	{
+		return !HiddenColorWheels[ColorWheelIndex];
+	}
+
+	return false;
+}
+
+void SColorGradingColorWheelPanel::OnColorGradingGroupSelectionChanged()
+{
+	Refresh();
+}
+
+void SColorGradingColorWheelPanel::OnColorGradingElementSelectionChanged()
+{
+	if (const FColorGradingEditorDataModel::FColorGradingElement* ColorGradingElement = ColorGradingDataModel->GetSelectedColorGradingElement())
+	{
+		FillColorWheels(*ColorGradingElement);
+	}
+	else
+	{
+		ClearColorWheels();
+	}
+}
+
+void SColorGradingColorWheelPanel::OnColorGradingElementCheckedChanged(ECheckBoxState State, FText ElementName)
+{
+	if (State == ECheckBoxState::Checked && ColorGradingDataModel)
+	{
+		if (FColorGradingEditorDataModel::FColorGradingGroup* ColorGradingGroup = ColorGradingDataModel->GetSelectedColorGradingGroup())
+		{
+			int32 ColorGradingElementIndex = ColorGradingGroup->ColorGradingElements.IndexOfByPredicate([=](const FColorGradingEditorDataModel::FColorGradingElement& Element)
+			{
+				return Element.DisplayName.CompareTo(ElementName) == 0;
+			});
+
+			ColorGradingDataModel->SetSelectedColorGradingElement(ColorGradingElementIndex);
+		}
+	}
+}
+
+ECheckBoxState SColorGradingColorWheelPanel::IsColorGradingElementSelected(FText ElementName) const
+{
+	if (ColorGradingDataModel)
+	{
+		if (const FColorGradingEditorDataModel::FColorGradingElement* ColorGradingElement = ColorGradingDataModel->GetSelectedColorGradingElement())
+		{
+			if (ColorGradingElement->DisplayName.CompareTo(ElementName) == 0)
+			{
+				return ECheckBoxState::Checked;
+			}
+		}
+	}
+
+	return ECheckBoxState::Unchecked;
+}
+
+EVisibility SColorGradingColorWheelPanel::GetColorWheelPanelVisibility() const
+{
+	bool bHasObjects = ColorGradingDataModel && ColorGradingDataModel->GetPropertyRowGenerator()->GetSelectedObjects().Num() > 0;
+	return bHasObjects ? EVisibility::Visible : EVisibility::Collapsed;
+}
+
+EVisibility SColorGradingColorWheelPanel::GetColorWheelVisibility(int32 ColorWheelIndex) const
+{
+	bool bIsHidden = HiddenColorWheels[ColorWheelIndex];
+	return bIsHidden ? EVisibility::Collapsed : EVisibility::Visible;
+}
+
+ECheckBoxState SColorGradingColorWheelPanel::IsColorDisplayModeChecked(EColorGradingColorDisplayMode InColorDisplayMode) const
+{
+	bool bIsModeSelected = InColorDisplayMode == ColorDisplayMode;
+	return bIsModeSelected ? ECheckBoxState::Checked : ECheckBoxState::Unchecked;
+}
+
+void SColorGradingColorWheelPanel::OnColorDisplayModeCheckedChanged(ECheckBoxState State, EColorGradingColorDisplayMode InColorDisplayMode)
+{
+	if (State == ECheckBoxState::Checked)
+	{
+		ColorDisplayMode = InColorDisplayMode;
+	}
+}
+
+FText SColorGradingColorWheelPanel::GetColorDisplayModeLabel(EColorGradingColorDisplayMode InColorDisplayMode) const
+{
+	FText Text;
+
+	switch (InColorDisplayMode)
+	{
+	case EColorGradingColorDisplayMode::RGB: Text = LOCTEXT("ColorWheel_RGBColorDisplayModeLabel", "RGB"); break;
+	case EColorGradingColorDisplayMode::HSV: Text = LOCTEXT("ColorWheel_HSVColorDisplayModeLabel", "HSV"); break;
+	}
+
+	return Text;
+}
+
+FText SColorGradingColorWheelPanel::GetColorDisplayModeToolTip(EColorGradingColorDisplayMode InColorDisplayMode) const
+{
+	FText Text;
+
+	switch (InColorDisplayMode)
+	{
+	case EColorGradingColorDisplayMode::RGB: Text = LOCTEXT("ColorWheel_RGBColorDisplayModeToolTip", "Change to RGB color mode"); break;
+	case EColorGradingColorDisplayMode::HSV: Text = LOCTEXT("ColorWheel_HSVColorDisplayModeToolTip", "Change to HSV color mode"); break;
+	}
+
+	return Text;
+}
+
+#undef LOCTEXT_NAMESPACE
\ No newline at end of file
diff --git a/Engine/Plugins/Editor/ColorGrading/Source/ColorGradingEditor/Private/SColorGradingColorWheelPanel.h b/Engine/Plugins/Editor/ColorGrading/Source/ColorGradingEditor/Private/SColorGradingColorWheelPanel.h
new file mode 100644
--- /dev/null
+++ b/Engine/Plugins/Editor/ColorGrading/Source/ColorGradingEditor/Private/SColorGradingColorWheelPanel.h
@@ -0,0 +1,109 @@
+// Copyright Epic Games, Inc. All Rights Reserved.
+
+#pragma once
+
+#include "CoreMinimal.h"
+#include "ColorGradingPanelState.h"
+#include "ColorGradingEditorDataModel.h"
+#include "SColorGradingColorWheel.h"
+#include "Styling/SlateTypes.h"
+#include "Widgets/DeclarativeSyntaxSupport.h"
+#include "Widgets/SCompoundWidget.h"
+
+class IDetailTreeNode;
+class IPropertyRowGenerator;
+class SBox;
+class SColorGradingDetailView;
+class SHorizontalBox;
+struct FColorGradingPanelState;
+
+/** A panel that contains up to five color wheels (for saturation, contrast, gamma, gain, and offset) as well as a details view for extra, non-color properties */
+class SColorGradingColorWheelPanel : public SCompoundWidget
+{
+private:
+	/** The number of color wheels the color wheel panel displays (for saturation, contrast, gamma, gain, and offset) */
+	static const uint32 NumColorWheels = 5;
+
+public:
+	virtual ~SColorGradingColorWheelPanel() override;
+
+	SLATE_BEGIN_ARGS(SColorGradingColorWheelPanel) {}
+		SLATE_ARGUMENT(TSharedPtr<FColorGradingEditorDataModel>, ColorGradingDataModelSource)
+	SLATE_END_ARGS()
+
+	void Construct(const FArguments& InArgs);
+
+	/** Regenerates the color wheel panel from the current state of the data model source */
+	void Refresh();
+
+	/** Adds the state of the color wheel panel to the specified drawer state */
+	void GetPanelState(FColorGradingPanelState& OutPanelState);
+
+	/** Sets the state of the color wheel panel from the specified drawer state */
+	void SetPanelState(const FColorGradingPanelState& InPanelState);
+
+private:
+	void BindCommands();
+
+	TSharedRef<SWidget> MakeColorDisplayModeCheckbox();
+	TSharedRef<SWidget> MakeSettingsMenu();
+
+	void FillColorGradingGroupProperty(const FColorGradingEditorDataModel::FColorGradingGroup& ColorGradingGroup);
+	void ClearColorGradingGroupProperty();
+
+	void FillColorGradingElementsToolBar(const TArray<FColorGradingEditorDataModel::FColorGradingElement>& ColorGradingElements);
+	void ClearColorGradingElementsToolBar();
+
+	void FillColorWheels(const FColorGradingEditorDataModel::FColorGradingElement& ColorGradingElement);
+	void ClearColorWheels();
+
+	TSharedRef<SWidget> CreateColorWheelHeaderWidget(const TSharedPtr<IPropertyHandle>& ColorPropertyHandle);
+	TSharedRef<SWidget> CreateColorPropertyExtensions(const TSharedPtr<IPropertyHandle>& ColorPropertyHandle, const TSharedPtr<IDetailTreeNode>& DetailTreeNode);
+
+	bool FilterDetailTreeNode(const TSharedRef<IDetailTreeNode>& InDetailTreeNode);
+
+	void SetColorWheelOrientation(EOrientation NewOrientation);
+	bool IsColorWheelOrientationSelected(EOrientation Orientation) const;
+
+	void ToggleColorWheelVisible(int32 ColorWheelIndex);
+	bool IsColorWheelVisible(int32 ColorWheelIndex);
+
+	void OnColorGradingGroupSelectionChanged();
+	void OnColorGradingElementSelectionChanged();
+
+	void OnColorGradingElementCheckedChanged(ECheckBoxState State, FText ElementName);
+
+	ECheckBoxState IsColorGradingElementSelected(FText ElementName) const;
+	EVisibility GetColorWheelPanelVisibility() const;
+	EVisibility GetColorWheelVisibility(int32 ColorWheelIndex) const;
+
+	EColorGradingColorDisplayMode GetColorDisplayMode() const { return ColorDisplayMode; }
+	ECheckBoxState IsColorDisplayModeChecked(EColorGradingColorDisplayMode InColorDisplayMode) const;
+	void OnColorDisplayModeCheckedChanged(ECheckBoxState State, EColorGradingColorDisplayMode InColorDisplayMode);
+	FText GetColorDisplayModeLabel(EColorGradingColorDisplayMode InColorDisplayMode) const;
+	FText GetColorDisplayModeToolTip(EColorGradingColorDisplayMode InColorDisplayMode) const;
+
+private:
+	/** The color grading data model that the panel is displaying */
+	TSharedPtr<FColorGradingEditorDataModel> ColorGradingDataModel;
+
+	TSharedPtr<SBox> ColorGradingGroupPropertyBox;
+	TSharedPtr<SHorizontalBox> ColorGradingElementsToolBarBox;
+
+	TArray<TSharedPtr<SColorGradingColorWheel>> ColorWheels;
+	TArray<bool> HiddenColorWheels;
+
+	TSharedPtr<SColorGradingDetailView> DetailView;
+
+	/** Commands used by the color wheel panel */
+	TSharedPtr<FUICommandList> CommandList;
+
+	/** The currently selected color grading group */
+	int32 SelectedColorGradingGroup = INDEX_NONE;
+
+	/** The current color display mode for the color wheels */
+	EColorGradingColorDisplayMode ColorDisplayMode = EColorGradingColorDisplayMode::RGB;
+
+	/** The current orientation for the color wheels */
+	EOrientation ColorWheelOrientation = EOrientation::Orient_Vertical;
+};
\ No newline at end of file
diff --git a/Engine/Plugins/Editor/ColorGrading/Source/ColorGradingEditor/Private/SColorGradingObjectList.cpp b/Engine/Plugins/Editor/ColorGrading/Source/ColorGradingEditor/Private/SColorGradingObjectList.cpp
new file mode 100644
--- /dev/null
+++ b/Engine/Plugins/Editor/ColorGrading/Source/ColorGradingEditor/Private/SColorGradingObjectList.cpp
@@ -0,0 +1,242 @@
+// Copyright Epic Games, Inc. All Rights Reserved.
+
+#include "SColorGradingObjectList.h"
+
+#include "ClassIconFinder.h"
+#include "Components/ActorComponent.h"
+#include "GameFramework/Actor.h"
+#include "Styling/AppStyle.h"
+#include "Widgets/Images/SImage.h"
+#include "Widgets/Input/SCheckBox.h"
+#include "Widgets/Views/SListView.h"
+#include "Widgets/Views/STableRow.h"
+
+#define LOCTEXT_NAMESPACE "ColorGrading"
+
+bool FColorGradingListItem::operator<(const FColorGradingListItem& Other) const
+{
+	FString ThisName = TEXT("");
+	FString OtherName = TEXT("");
+
+	if (Component.IsValid())
+	{
+		ThisName = Component->GetName();
+	}
+	else if (Actor.IsValid())
+	{
+		ThisName = Actor->GetActorLabel();
+	}
+
+	if (Other.Component.IsValid())
+	{
+		OtherName = Other.Component->GetName();
+	}
+	else if (Other.Actor.IsValid())
+	{
+		OtherName = Other.Actor->GetActorLabel();
+	}
+
+	return ThisName < OtherName;
+}
+
+namespace ColorGradingObjectListColumnNames
+{
+	const static FName ItemEnabled(TEXT("ItemEnabled"));
+	const static FName ItemLabel(TEXT("ItemLabel"));
+};
+
+class SColorGradingListItemRow : public SMultiColumnTableRow<FColorGradingListItemRef>
+{
+public:
+	SLATE_BEGIN_ARGS(SColorGradingListItemRow) {}
+	SLATE_END_ARGS()
+
+	void Construct(const FArguments& InArgs, const TSharedRef<STableViewBase>& InOwnerTable, const FColorGradingListItemRef InListItem)
+	{
+		ListItem = InListItem;
+
+		SMultiColumnTableRow<FColorGradingListItemRef>::Construct(
+			SMultiColumnTableRow<FColorGradingListItemRef>::FArguments()
+			.Style(&FAppStyle::Get().GetWidgetStyle<FTableRowStyle>("SceneOutliner.TableViewRow"))
+			, InOwnerTable);
+	}
+
+	virtual TSharedRef<SWidget> GenerateWidgetForColumn(const FName& ColumnName) override
+	{
+		if (ColumnName == ColorGradingObjectListColumnNames::ItemEnabled)
+		{
+			return SNew(SCheckBox)
+				.IsChecked(this, &SColorGradingListItemRow::GetCheckBoxState)
+				.OnCheckStateChanged(this, &SColorGradingListItemRow::OnCheckBoxStateChanged);
+		}
+		else if (ColumnName == ColorGradingObjectListColumnNames::ItemLabel)
+		{
+			return SNew(SHorizontalBox)
+
+				+ SHorizontalBox::Slot()
+				.AutoWidth()
+				.Padding(FMargin(6, 1))
+				.VAlign(VAlign_Center)
+				[
+					SNew(SBox)
+					.WidthOverride(16)
+					.HeightOverride(16)
+					[
+						SNew(SImage)
+						.ColorAndOpacity(FSlateColor::UseForeground())
+						.Image(this, &SColorGradingListItemRow::GetItemIcon)
+					]
+				]
+
+				+ SHorizontalBox::Slot()
+				.AutoWidth()
+				.VAlign(VAlign_Center)
+				.HAlign(HAlign_Left)
+				[
+					SNew(STextBlock)
+					.Text(this, &SColorGradingListItemRow::GetItemLabel)
+				];
+		}
+
+		return SNullWidget::NullWidget;
+	}
+
+private:
+	FText GetItemLabel() const
+	{
+		FString ItemName = TEXT("");
+
+		if (ListItem.IsValid())
+		{
+			if (ListItem->Component.IsValid())
+			{
+				ItemName = ListItem->Component->GetName();
+			}
+			else if (ListItem->Actor.IsValid())
+			{
+				ItemName = ListItem->Actor->GetActorLabel();
+			}
+		}
+
+		return FText::FromString(*ItemName);
+	}
+
+	const FSlateBrush* GetItemIcon() const
+	{
+		if (ListItem.IsValid())
+		{
+			if (ListItem->Component.IsValid())
+			{
+				return FSlateIconFinder::FindIconBrushForClass(ListItem->Component->GetClass(), TEXT("SCS.Component"));
+			}
+			else if (ListItem->Actor.IsValid())
+			{
+				return FClassIconFinder::FindIconForActor(ListItem->Actor);
+			}
+		}
+
+		return nullptr;
+	}
+
+	EVisibility GetCheckBoxVisibility() const
+	{
+		if (ListItem.IsValid())
+		{
+			bool bEnabledIsSet = ListItem->IsItemEnabled.IsSet();
+			return bEnabledIsSet ? EVisibility::Visible : EVisibility::Collapsed;
+		}
+
+		return EVisibility::Collapsed;
+	}
+
+	ECheckBoxState GetCheckBoxState() const
+	{
+		if (ListItem.IsValid())
+		{
+			bool bIsEnabled = ListItem->IsItemEnabled.Get(false);
+			return bIsEnabled ? ECheckBoxState::Checked : ECheckBoxState::Unchecked;
+		}
+
+		return ECheckBoxState::Unchecked;
+	}
+
+	void OnCheckBoxStateChanged(ECheckBoxState NewState)
+	{
+		if (ListItem.IsValid())
+		{
+			const bool bNewEnabledState = NewState == ECheckBoxState::Checked ? true : false;
+			ListItem->OnItemEnabledChanged.ExecuteIfBound(ListItem, bNewEnabledState);
+		}
+	}
+
+private:
+	FColorGradingListItemRef ListItem;
+};
+
+void SColorGradingObjectList::Construct(const FArguments& InArgs)
+{
+	OnSelectionChangedDelegate = InArgs._OnSelectionChanged;
+
+	RefreshList();
+
+	ChildSlot
+	[
+		SNew(SVerticalBox)
+
+		+ SVerticalBox::Slot()
+		.FillHeight(1.0f)
+		[
+			SAssignNew(ListView, SListView<FColorGradingListItemRef>)
+			.ListItemsSource(InArgs._ColorGradingItemsSource)
+			.SelectionMode(ESelectionMode::Single)
+			.OnGenerateRow(this, &SColorGradingObjectList::GenerateListItemRow)
+			.OnSelectionChanged(this, &SColorGradingObjectList::OnSelectionChanged)
+			.HeaderRow
+			(
+				SNew(SHeaderRow)
+
+				+ SHeaderRow::Column(ColorGradingObjectListColumnNames::ItemEnabled)
+				.DefaultLabel(FText::GetEmpty())
+				.FixedWidth(24.0f)
+				.HAlignCell(HAlign_Center)
+				.VAlignCell(VAlign_Center)
+				.DefaultTooltip(LOCTEXT("ItemEnabledColumnTooltip", "Color Grading Enabled"))
+
+				+ SHeaderRow::Column(ColorGradingObjectListColumnNames::ItemLabel)
+				.DefaultLabel(LOCTEXT("ColorGradingPanelList_ItemLabelHeader", "Item Label"))
+				.FillWidth(1.0f)
+			)
+		]
+	];
+}
+
+void SColorGradingObjectList::RefreshList()
+{
+	if (ListView.IsValid())
+	{
+		ListView->RebuildList();
+	}
+}
+
+TArray<FColorGradingListItemRef> SColorGradingObjectList::GetSelectedItems()
+{
+	return ListView->GetSelectedItems();
+}
+
+void SColorGradingObjectList::SetSelectedItems(const TArray<FColorGradingListItemRef>& InSelectedItems)
+{
+	ListView->ClearSelection();
+	ListView->SetItemSelection(InSelectedItems, true);
+}
+
+TSharedRef<ITableRow> SColorGradingObjectList::GenerateListItemRow(FColorGradingListItemRef Item, const TSharedRef<STableViewBase>& OwnerTable)
+{
+	return SNew(SColorGradingListItemRow, OwnerTable, Item);
+}
+
+void SColorGradingObjectList::OnSelectionChanged(FColorGradingListItemRef SelectedItem, ESelectInfo::Type SelectInfo)
+{
+	OnSelectionChangedDelegate.ExecuteIfBound(SharedThis<SColorGradingObjectList>(this), SelectedItem, SelectInfo);
+}
+
+#undef LOCTEXT_NAMESPACE
\ No newline at end of file
diff --git a/Engine/Plugins/Editor/ColorGrading/Source/ColorGradingEditor/Private/SColorGradingObjectList.h b/Engine/Plugins/Editor/ColorGrading/Source/ColorGradingEditor/Private/SColorGradingObjectList.h
new file mode 100644
--- /dev/null
+++ b/Engine/Plugins/Editor/ColorGrading/Source/ColorGradingEditor/Private/SColorGradingObjectList.h
@@ -0,0 +1,56 @@
+// Copyright Epic Games, Inc. All Rights Reserved.
+
+#pragma once
+
+#include "CoreMinimal.h"
+#include "Components/ActorComponent.h"
+#include "Framework/SlateDelegates.h"
+#include "GameFramework/Actor.h"
+#include "Widgets/DeclarativeSyntaxSupport.h"
+#include "Widgets/SCompoundWidget.h"
+
+#include "ColorGradingListItem.h"
+
+class ITableRow;
+class STableViewBase;
+
+template<class T>
+class SListView;
+
+/** Displays a list of color gradable items */
+class SColorGradingObjectList : public SCompoundWidget
+{
+public:
+	DECLARE_DELEGATE_ThreeParams(FOnSelectionChanged, TSharedRef<SColorGradingObjectList>, FColorGradingListItemRef, ESelectInfo::Type);
+
+public:
+	SLATE_BEGIN_ARGS(SColorGradingObjectList) {}
+		SLATE_ARGUMENT(const TArray<FColorGradingListItemRef>*, ColorGradingItemsSource)
+		SLATE_EVENT(FOnSelectionChanged, OnSelectionChanged)
+	SLATE_END_ARGS()
+
+	void Construct(const FArguments& InArgs);
+
+	/** Refreshes the list, updating the UI to reflect the current state of the source items list*/
+	void RefreshList();
+
+	/** Gets a list of currently selected items */
+	TArray<FColorGradingListItemRef> GetSelectedItems();
+
+	/** Selects the specified list of items */
+	void SetSelectedItems(const TArray<FColorGradingListItemRef>& InSelectedItems);
+
+private:
+	/** Generates the table row widget for the specified list item */
+	TSharedRef<ITableRow> GenerateListItemRow(FColorGradingListItemRef Item, const TSharedRef<STableViewBase>& OwnerTable);
+
+	/** Raised when the internal list view's selection has changed */
+	void OnSelectionChanged(FColorGradingListItemRef SelectedItem, ESelectInfo::Type SelectInfo);
+
+private:
+	/** Internal list view used to display the list of color gradable items */
+	TSharedPtr<SListView<FColorGradingListItemRef>> ListView;
+
+	/** A delegate that is raised when the list of selected items is changed */
+	FOnSelectionChanged OnSelectionChangedDelegate;
+};
\ No newline at end of file
diff --git a/Engine/Plugins/Editor/ColorGrading/Source/ColorGradingEditor/Private/SColorGradingPanel.cpp b/Engine/Plugins/Editor/ColorGrading/Source/ColorGradingEditor/Private/SColorGradingPanel.cpp
new file mode 100644
--- /dev/null
+++ b/Engine/Plugins/Editor/ColorGrading/Source/ColorGradingEditor/Private/SColorGradingPanel.cpp
@@ -0,0 +1,766 @@
+// Copyright Epic Games, Inc. All Rights Reserved.
+
+#include "SColorGradingPanel.h"
+
+#include "ColorGradingCommands.h"
+#include "ColorGradingEditorDataModel.h"
+#include "IColorGradingEditor.h"
+#include "SColorGradingColorWheelPanel.h"
+#include "SColorGradingObjectList.h"
+
+#include "ColorCorrectRegion.h"
+#include "Editor.h"
+#include "Engine/Blueprint.h"
+#include "Engine/Engine.h"
+#include "EngineUtils.h"
+#include "Framework/Commands/GenericCommands.h"
+#include "Framework/MultiBox/MultiBoxBuilder.h"
+#include "LevelEditorSubsystem.h"
+#include "Widgets/Input/SButton.h"
+#include "Widgets/Layout/SBox.h"
+#include "Widgets/Layout/SExpandableArea.h"
+#include "Widgets/Layout/SScrollBox.h"
+#include "Widgets/Layout/SSeparator.h"
+#include "Widgets/Layout/SSpacer.h"
+#include "Widgets/Layout/SSplitter.h"
+#include "Widgets/Text/SInlineEditableTextBlock.h"
+
+#define LOCTEXT_NAMESPACE "ColorGradingEditor"
+
+SColorGradingPanel::~SColorGradingPanel()
+{
+	for (const FColorGradingListItemRef& ColorGradingItem : ColorGradingItemList)
+	{
+		if (ColorGradingItem->Component.IsValid())
+		{
+			UnbindBlueprintCompiledDelegate(ColorGradingItem->Component->GetClass());
+		}
+
+		if (ColorGradingItem->Actor.IsValid())
+		{
+			UnbindBlueprintCompiledDelegate(ColorGradingItem->Actor->GetClass());
+		}
+	}
+
+	FCoreUObjectDelegates::OnObjectsReplaced.RemoveAll(this);
+
+	if (GEngine)
+	{
+		GEngine->OnLevelActorAdded().RemoveAll(this);
+		GEngine->OnLevelActorDeleted().RemoveAll(this);
+	}
+
+	if (GEditor)
+	{
+		GEditor->UnregisterForUndo(this);
+	}
+}
+
+void SColorGradingPanel::Construct(const FArguments& InArgs)
+{
+	bIsInDrawer = InArgs._IsInDrawer;
+	DockCallback = InArgs._OnDocked;
+	OverrideWorld = InArgs._OverrideWorld;
+	ActorFilter = InArgs._ActorFilter;
+
+	ColorGradingDataModel = MakeShared<FColorGradingEditorDataModel>();
+	ColorGradingDataModel->OnDataModelGenerated().AddSP(this, &SColorGradingPanel::OnColorGradingDataModelGenerated);
+
+	FCoreUObjectDelegates::OnObjectsReplaced.AddSP(this, &SColorGradingPanel::OnObjectsReplaced);
+	GEngine->OnLevelActorAdded().AddSP(this, &SColorGradingPanel::OnLevelActorAdded);
+	GEngine->OnLevelActorDeleted().AddSP(this, &SColorGradingPanel::OnLevelActorDeleted);
+	GEngine->OnWorldAdded().AddSP(this, &SColorGradingPanel::OnWorldAdded);
+	GEngine->OnWorldDestroyed().AddSP(this, &SColorGradingPanel::OnWorldDestroyed);
+
+	GEditor->RegisterForUndo(this);
+
+	RefreshColorGradingList();
+
+	ChildSlot
+	[
+		SNew(SBorder)
+		.BorderImage(FAppStyle::Get().GetBrush("Brushes.Panel"))
+		.HAlign(HAlign_Fill)
+		.VAlign(VAlign_Fill)
+		.Padding(FMargin(0.0f, 0.0f))
+		[
+			// Splitter to divide the object list and the color panel
+			SNew(SSplitter)
+			.Orientation(Orient_Horizontal)
+			.PhysicalSplitterHandleSize(2.0f)
+
+			// Splitter slot for object list
+			+SSplitter::Slot()
+			.Value(0.12f)
+			[
+				SNew(SBox)
+				.Padding(FMargin(4.f))
+				[
+					SNew(SBorder)
+					.Padding(FMargin(0.0f))
+					.BorderImage(FAppStyle::GetBrush("Brushes.Recessed"))
+					[
+						SNew(SScrollBox)
+						+SScrollBox::Slot()
+						[
+							SNew(SExpandableArea)
+							.BorderImage(FAppStyle::Get().GetBrush("Brushes.Header"))
+							.BodyBorderImage(FAppStyle::Get().GetBrush("Brushes.Recessed"))
+							.HeaderPadding(FMargin(4.0f, 2.0f))
+							.InitiallyCollapsed(false)
+							.AllowAnimatedTransition(false)
+							.Visibility_Lambda([this]() { return ColorGradingItemList.Num() ? EVisibility::Visible : EVisibility::Collapsed; })
+							.HeaderContent()
+							[
+								SNew(SBox)
+								.HeightOverride(24.0f)
+								.VAlign(VAlign_Center)
+								[
+									SNew(STextBlock)
+										.Text(LOCTEXT("ColorGradingObjectListLabel", "Objects"))
+										.TextStyle(FAppStyle::Get(), "ButtonText")
+										.Font(FAppStyle::Get().GetFontStyle("NormalFontBold"))
+								]
+							]
+							.BodyContent()
+							[
+								SAssignNew(ColorGradingObjectListView, SColorGradingObjectList)
+								.ColorGradingItemsSource(&ColorGradingItemList)
+								.OnSelectionChanged(this, &SColorGradingPanel::OnListSelectionChanged)
+							]
+						]
+					]
+				]
+			]
+
+			// Splitter slot for color grading controls/details
+			+SSplitter::Slot()
+			.Value(0.88f)
+			[
+				SNew(SVerticalBox)
+
+				// Toolbar slot for the main drawer toolbar
+				+SVerticalBox::Slot()
+				.AutoHeight()
+				.Padding(0, 0, 0, 0)
+				[
+					SNew(SBorder)
+					.Padding(FMargin(3))
+					.BorderImage(bIsInDrawer ? FStyleDefaults::GetNoBrush() : FAppStyle::Get().GetBrush("Brushes.Panel"))
+					[
+						SNew(SBox)
+						.HeightOverride(28.0f)
+						[
+							SNew(SHorizontalBox)
+
+							// Slot for the color grading group toolbar
+							+SHorizontalBox::Slot()
+							.AutoWidth()
+							.HAlign(HAlign_Left)
+							.VAlign(VAlign_Center)
+							[
+								SAssignNew(ColorGradingGroupToolBarBox, SHorizontalBox)
+								.Visibility(this, &SColorGradingPanel::GetColorGradingGroupToolBarVisibility)
+							]
+
+
+							+SHorizontalBox::Slot()
+							.FillWidth(1.0f)
+							[
+								SNew(SSpacer)
+							]
+
+							// Slot for the "Dock in Layout" button
+							+SHorizontalBox::Slot()
+							.AutoWidth()
+							.HAlign(HAlign_Right)
+							.VAlign(VAlign_Center)
+							[
+								CreateDockInLayoutButton()
+							]
+						]
+					]
+				]
+
+				+SVerticalBox::Slot()
+				.AutoHeight()
+				[
+					SNew(SSeparator)
+					.Thickness(2.0f)
+				]
+
+				// Slot for the color panel
+				+SVerticalBox::Slot()
+				[
+					SNew(SBorder)
+					.BorderImage(FAppStyle::Get().GetBrush("Brushes.Panel"))
+					.Padding(FMargin(2.0f, 2.0f, 2.0f, 0.0f))
+					[
+						SAssignNew(ColorWheelPanel, SColorGradingColorWheelPanel)
+						.ColorGradingDataModelSource(ColorGradingDataModel)
+					]
+				]
+			]
+		]
+	];
+}
+
+void SColorGradingPanel::Tick(const FGeometry& AllottedGeometry, const double InCurrentTime, const float InDeltaTime)
+{
+	SCompoundWidget::Tick(AllottedGeometry, InCurrentTime, InDeltaTime);
+
+	if (bRefreshOnNextTick)
+	{
+		Refresh();
+
+		bRefreshOnNextTick = false;
+	}
+}
+
+void SColorGradingPanel::Refresh()
+{
+	FColorGradingPanelState PanelState;
+	GetPanelState(PanelState);
+
+	ColorGradingDataModel->Reset();
+
+	RefreshColorGradingList();
+
+	if (ColorWheelPanel.IsValid())
+	{
+		ColorWheelPanel->Refresh();
+	}
+
+	SetPanelState(PanelState);
+}
+
+void SColorGradingPanel::GetPanelState(FColorGradingPanelState& OutPanelState) const
+{
+	ColorGradingDataModel->GetPanelState(OutPanelState);
+
+	if (ColorWheelPanel.IsValid())
+	{
+		ColorWheelPanel->GetPanelState(OutPanelState);
+	}
+
+	if (ColorGradingObjectListView.IsValid())
+	{
+		TArray<FColorGradingListItemRef> SelectedItems = ColorGradingObjectListView->GetSelectedItems();
+
+		for (const FColorGradingListItemRef& SelectedItem : SelectedItems)
+		{
+			if (SelectedItem.IsValid())
+			{
+				if (SelectedItem->Component.IsValid())
+				{
+					OutPanelState.SelectedObjects.Add(SelectedItem->Component);
+				}
+				else if (SelectedItem->Actor.IsValid())
+				{
+					OutPanelState.SelectedObjects.Add(SelectedItem->Actor);
+				}
+			}
+		}
+	}
+}
+
+void SColorGradingPanel::SetPanelState(const FColorGradingPanelState& InPanelState)
+{
+	TArray<FColorGradingListItemRef> ItemsToSelect;
+
+	for (const TWeakObjectPtr<UObject>& SelectedObject : InPanelState.SelectedObjects)
+	{
+		if (SelectedObject.IsValid())
+		{
+			auto FindColorGradingItem = [&SelectedObject](const FColorGradingListItemRef& ColorGradingItem)
+			{
+				if (SelectedObject->IsA<AActor>())
+				{
+					return ColorGradingItem->Actor == SelectedObject && ColorGradingItem->Component == nullptr;
+				}
+
+				return ColorGradingItem->Actor == SelectedObject || ColorGradingItem->Component == SelectedObject;
+			};
+
+			if (FColorGradingListItemRef* FoundItem = ColorGradingItemList.FindByPredicate(FindColorGradingItem))
+			{
+				ItemsToSelect.Add(*FoundItem);
+				break;
+			}
+		}
+	}
+
+	if (!ItemsToSelect.IsEmpty() && ColorGradingObjectListView.IsValid())
+	{
+		ColorGradingObjectListView->SetSelectedItems(ItemsToSelect);
+	}
+
+	ColorGradingDataModel->SetPanelState(InPanelState);
+
+	if (ColorWheelPanel.IsValid())
+	{
+		ColorWheelPanel->SetPanelState(InPanelState);
+	}
+}
+
+void SColorGradingPanel::SetSelectedObjects(const TArray<UObject*>& Objects)
+{
+	ColorGradingDataModel->SetObjects(Objects);
+
+	FColorGradingPanelState PanelState;
+	GetPanelState(PanelState);
+
+	PanelState.SelectedObjects.Empty(Objects.Num());
+	PanelState.SelectedObjects.Append(Objects);
+
+	SetPanelState(PanelState);
+}
+
+TSharedRef<SWidget> SColorGradingPanel::CreateDockInLayoutButton()
+{
+	if (bIsInDrawer && DockCallback.IsBound())
+	{
+		return SNew(SButton)
+			.ButtonStyle(FAppStyle::Get(), "SimpleButton")
+			.ToolTipText(LOCTEXT("DockInLayout_Tooltip", "Docks this panel in the current window, copying all settings from the drawer.\nThe drawer will still be usable."))
+			.OnClicked(this, &SColorGradingPanel::DockInLayout)
+			[
+				SNew(SHorizontalBox)
+
+				+ SHorizontalBox::Slot()
+				.AutoWidth()
+				.VAlign(VAlign_Center)
+				.Padding(4.0, 0.0f)
+				[
+					SNew(SImage)
+					.ColorAndOpacity(FSlateColor::UseForeground())
+					.Image(FAppStyle::Get().GetBrush("Icons.Layout"))
+				]
+
+				+ SHorizontalBox::Slot()
+				.VAlign(VAlign_Center)
+				.Padding(4.0, 0.0f)
+				[
+					SNew(STextBlock)
+					.Text(LOCTEXT("DockInLayout", "Dock in Layout"))
+					.ColorAndOpacity(FSlateColor::UseForeground())
+				]
+			];
+	}
+
+	return SNullWidget::NullWidget;
+}
+
+UWorld* SColorGradingPanel::GetWorld()
+{
+	if (OverrideWorld.IsSet() && OverrideWorld.Get())
+	{
+		return OverrideWorld.Get();
+	}
+
+	ULevelEditorSubsystem* LevelEditorSubsystem = GEditor->GetEditorSubsystem<ULevelEditorSubsystem>();
+	if (!LevelEditorSubsystem)
+	{
+		return nullptr;
+	}
+
+	ULevel* Level = LevelEditorSubsystem->GetCurrentLevel();
+	if (!Level)
+	{
+		return nullptr;
+	}
+
+	return Level->GetWorld();
+}
+
+void SColorGradingPanel::BindBlueprintCompiledDelegate(const UClass* Class)
+{
+	if (UBlueprint* Blueprint = UBlueprint::GetBlueprintFromClass(Class))
+	{
+		if (!Blueprint->OnCompiled().IsBoundToObject(this))
+		{
+			Blueprint->OnCompiled().AddSP(this, &SColorGradingPanel::OnBlueprintCompiled);
+		}
+	}
+}
+
+void SColorGradingPanel::UnbindBlueprintCompiledDelegate(const UClass* Class)
+{
+	if (UBlueprint* Blueprint = UBlueprint::GetBlueprintFromClass(Class))
+	{
+		Blueprint->OnCompiled().RemoveAll(this);
+	}
+}
+
+void SColorGradingPanel::RefreshColorGradingList()
+{
+	for (const FColorGradingListItemRef& Item : ColorGradingItemList)
+	{
+		if (Item->Component.IsValid())
+		{
+			UnbindBlueprintCompiledDelegate(Item->Component->GetClass());
+		}
+
+		if (Item->Actor.IsValid())
+		{
+			UnbindBlueprintCompiledDelegate(Item->Actor->GetClass());
+		}
+	}
+
+	ColorGradingItemList.Empty();
+
+	UWorld* World = GetWorld();
+	if (World && ColorGradingDataModel.IsValid())
+	{
+		// Sorter that sorts the list items alphabetically by their display name
+		auto AlphabeticalSort = [](const FColorGradingListItemRef& A, const FColorGradingListItemRef& B)
+		{
+			if (A.IsValid() && B.IsValid())
+			{
+				return *A < *B;
+			}
+			else
+			{
+				return false;
+			}
+		};
+
+		// Add all relevant actors
+		const TSet<TSubclassOf<AActor>>& ListItemClasses = FColorGradingListItem::GetActorClassesWithListItemGenerators();
+
+		for (TSubclassOf<AActor> ActorClass : ListItemClasses)
+		{
+			for (TActorIterator<AActor> ActorIter(World, ActorClass.Get()); ActorIter; ++ActorIter)
+			{
+				AActor* Actor = *ActorIter;
+
+				if (ActorFilter && !ActorFilter(Actor))
+				{
+					continue;
+				}
+
+				TArray<FColorGradingListItemRef> ListItems = FColorGradingListItem::GenerateColorGradingListItems(Actor);
+
+				if (ListItems.IsEmpty())
+				{
+					continue;
+				}
+
+				BindBlueprintCompiledDelegate(Actor->GetClass());
+				ColorGradingItemList.Append(ListItems);
+			}
+		}
+
+		ColorGradingItemList.Sort(AlphabeticalSort);
+	}
+
+	if (ColorGradingObjectListView.IsValid())
+	{
+		ColorGradingObjectListView->RefreshList();
+	}
+}
+
+void SColorGradingPanel::FillColorGradingGroupToolBar()
+{
+	if (ColorGradingGroupToolBarBox.IsValid())
+	{
+		ColorGradingGroupToolBarBox->ClearChildren();
+		ColorGradingGroupTextBlocks.Empty(ColorGradingDataModel->ColorGradingGroups.Num());
+
+		for (int32 Index = 0; Index < ColorGradingDataModel->ColorGradingGroups.Num(); ++Index)
+		{
+			TSharedPtr<SInlineEditableTextBlock> TextBlock = nullptr;
+
+			ColorGradingGroupToolBarBox->AddSlot()
+			.AutoWidth()
+			.Padding(2.0f, 0.0f, 2.0f, 0.0f)
+			[
+				SNew(SCheckBox)
+				.Style(FAppStyle::Get(), "DetailsView.SectionButton")
+				.OnCheckStateChanged(this, &SColorGradingPanel::OnColorGradingGroupCheckedChanged, Index)
+				.IsChecked(this, &SColorGradingPanel::IsColorGradingGroupSelected, Index)
+				.OnGetMenuContent(this, &SColorGradingPanel::GetColorGradingGroupMenuContent, Index)
+				[
+					SNew(SBox)
+					.HeightOverride(20.0)
+					[
+						SAssignNew(TextBlock, SInlineEditableTextBlock)
+						.Text(this, &SColorGradingPanel::GetColorGradingGroupDisplayName, Index)
+						.Font(this, &SColorGradingPanel::GetColorGradingGroupDisplayNameFont, Index)
+						.OnTextCommitted(this, &SColorGradingPanel::OnColorGradingGroupRenamed, Index)
+					]
+				]
+			];
+
+			ColorGradingGroupTextBlocks.Add(TextBlock);
+		}
+
+		if (ColorGradingDataModel->ColorGradingGroupToolBarWidget.IsValid())
+		{
+			ColorGradingGroupToolBarBox->AddSlot()
+			.AutoWidth()
+			.Padding(2.0f, 0.0f, 2.0f, 0.0f)
+			[
+				ColorGradingDataModel->ColorGradingGroupToolBarWidget.ToSharedRef()
+			];
+		}
+	}
+}
+
+EVisibility SColorGradingPanel::GetColorGradingGroupToolBarVisibility() const
+{
+	if (ColorGradingDataModel->bShowColorGradingGroupToolBar)
+	{
+		return EVisibility::Visible;
+	}
+
+	return EVisibility::Hidden;
+}
+
+ECheckBoxState SColorGradingPanel::IsColorGradingGroupSelected(int32 GroupIndex) const
+{
+	return ColorGradingDataModel->GetSelectedColorGradingGroupIndex() == GroupIndex ? ECheckBoxState::Checked : ECheckBoxState::Unchecked;
+}
+
+void SColorGradingPanel::OnColorGradingGroupCheckedChanged(ECheckBoxState State, int32 GroupIndex)
+{
+	if (State == ECheckBoxState::Checked)
+	{
+		ColorGradingDataModel->SetSelectedColorGradingGroup(GroupIndex);
+	}
+}
+
+FText SColorGradingPanel::GetColorGradingGroupDisplayName(int32 GroupIndex) const
+{
+	if (ColorGradingDataModel->ColorGradingGroups.IsValidIndex(GroupIndex))
+	{
+		FText DisplayName = ColorGradingDataModel->ColorGradingGroups[GroupIndex].DisplayName;
+		if (DisplayName.IsEmpty())
+		{
+			return LOCTEXT("ColorGradingGroupEmptyNameLabel", "Unnamed");
+		}
+
+		return DisplayName;
+	}
+
+	return FText::GetEmpty();
+}
+
+FSlateFontInfo SColorGradingPanel::GetColorGradingGroupDisplayNameFont(int32 GroupIndex) const
+{
+	if (ColorGradingDataModel->ColorGradingGroups.IsValidIndex(GroupIndex))
+	{
+		FText DisplayName = ColorGradingDataModel->ColorGradingGroups[GroupIndex].DisplayName;
+		if (DisplayName.IsEmpty())
+		{
+			return FAppStyle::Get().GetFontStyle("NormalFontItalic");
+		}
+	}
+
+	return FAppStyle::Get().GetFontStyle("NormalFont");
+}
+
+TSharedRef<SWidget> SColorGradingPanel::GetColorGradingGroupMenuContent(int32 GroupIndex)
+{
+	if (ColorGradingDataModel->ColorGradingGroups.IsValidIndex(GroupIndex))
+	{
+		const FColorGradingEditorDataModel::FColorGradingGroup& Group = ColorGradingDataModel->ColorGradingGroups[GroupIndex];
+
+		FMenuBuilder MenuBuilder(true, nullptr);
+
+		const FGenericCommands& GenericCommands = FGenericCommands::Get();
+
+		if (Group.bCanBeRenamed)
+		{
+			MenuBuilder.AddMenuEntry(
+				GenericCommands.Rename->GetLabel(),
+				GenericCommands.Rename->GetDescription(),
+				GenericCommands.Rename->GetIcon(),
+				FUIAction(FExecuteAction::CreateSP(this, &SColorGradingPanel::OnColorGradingGroupRequestRename, GroupIndex))
+			);
+		}
+
+		if (Group.bCanBeDeleted)
+		{
+			MenuBuilder.AddMenuEntry(
+				GenericCommands.Delete->GetLabel(),
+				GenericCommands.Delete->GetDescription(),
+				GenericCommands.Delete->GetIcon(),
+				FUIAction(FExecuteAction::CreateSP(this, &SColorGradingPanel::OnColorGradingGroupDeleted, GroupIndex))
+			);
+		}
+
+		return MenuBuilder.MakeWidget();
+	}
+
+	return SNullWidget::NullWidget;
+}
+
+void SColorGradingPanel::OnColorGradingGroupDeleted(int32 GroupIndex)
+{
+	// If the group being deleted is in front of the currently selected one, we want to make sure
+	// that the same group is selected even after the deletion, so preemptively adjust the
+	// currently selected group index
+	int32 SelectedGroupIndex = ColorGradingDataModel->GetSelectedColorGradingGroupIndex();
+	if (SelectedGroupIndex > GroupIndex)
+	{
+		ColorGradingDataModel->SetSelectedColorGradingGroup(SelectedGroupIndex - 1);
+	}
+
+	ColorGradingDataModel->OnColorGradingGroupDeleted().Broadcast(GroupIndex);
+}
+
+void SColorGradingPanel::OnColorGradingGroupRequestRename(int32 GroupIndex)
+{
+	if (ColorGradingGroupTextBlocks.IsValidIndex(GroupIndex) && ColorGradingGroupTextBlocks[GroupIndex].IsValid())
+	{
+		ColorGradingGroupTextBlocks[GroupIndex]->EnterEditingMode();
+	}
+}
+
+void SColorGradingPanel::OnColorGradingGroupRenamed(const FText& InText, ETextCommit::Type TextCommitType, int32 GroupIndex)
+{
+	ColorGradingDataModel->OnColorGradingGroupRenamed().Broadcast(GroupIndex, InText);
+}
+
+void SColorGradingPanel::OnObjectsReplaced(const TMap<UObject*, UObject*>& OldToNewInstanceMap)
+{
+	bool bNeedsFullRefresh = false;
+	bool bNeedsListRefresh = false;
+
+	const TArray<TWeakObjectPtr<UObject>> SelectedObjects = ColorGradingDataModel->GetObjects();
+
+	for (const TPair<UObject*, UObject*>& Pair : OldToNewInstanceMap)
+	{
+		if (Pair.Key && Pair.Value)
+		{
+			FColorGradingListItemRef* FoundColorGradingItemPtr = nullptr;
+
+			// Must use GetEvenIfUnreachable on the weak pointers here because most of the time, the objects being replaced have already been marked for GC, and TWeakObjectPtr
+			// will return nullptr from Get on GC-marked objects
+			FoundColorGradingItemPtr = ColorGradingItemList.FindByPredicate([&Pair](const FColorGradingListItemRef& ColorGradingItem)
+			{
+				return ColorGradingItem->Actor.GetEvenIfUnreachable() == Pair.Key || ColorGradingItem->Component.GetEvenIfUnreachable() == Pair.Key;
+			});
+
+			if (FoundColorGradingItemPtr)
+			{
+				FColorGradingListItemRef FoundColorGradingItem = *FoundColorGradingItemPtr;
+				if (FoundColorGradingItem->Actor.GetEvenIfUnreachable() == Pair.Key)
+				{
+					FoundColorGradingItem->Actor = Cast<AActor>(Pair.Value);
+				}
+				else if (FoundColorGradingItem->Component.GetEvenIfUnreachable() == Pair.Key)
+				{
+					FoundColorGradingItem->Component = Cast<UActorComponent>(Pair.Value);
+				}
+
+				bNeedsListRefresh = true;
+			}
+
+			if (SelectedObjects.Contains(Pair.Key))
+			{
+				bNeedsFullRefresh = true;
+			}
+		}
+	}
+
+	if (bNeedsFullRefresh)
+	{
+		Refresh();
+	}
+	else if (bNeedsListRefresh && ColorGradingObjectListView)
+	{
+		ColorGradingObjectListView->RefreshList();
+	}
+}
+
+void SColorGradingPanel::OnLevelActorAdded(AActor* Actor)
+{
+	// Only refresh when the actor being added is being added to the edited world
+	if (UWorld* World = GetWorld())
+	{
+		if (World == Actor->GetWorld())
+		{
+			if (Actor->IsA<APostProcessVolume>() || Actor->IsA<AColorCorrectRegion>())
+			{
+				// Wait to refresh, as this event can be fired off for several actors in a row in certain cases, such as when the root actor is recompiled after a property change
+				bRefreshOnNextTick = true;
+			}
+		}
+	}
+}
+
+void SColorGradingPanel::OnLevelActorDeleted(AActor* Actor)
+{
+	auto ContainsActorRef = [Actor](const FColorGradingListItemRef& ColorGradingItem)
+	{
+		return ColorGradingItem->Actor.GetEvenIfUnreachable() == Actor;
+	};
+
+	if (ColorGradingItemList.ContainsByPredicate(ContainsActorRef))
+	{
+		// Must wait for next tick to refresh because the actor has not actually been removed from the level at this point
+		bRefreshOnNextTick = true;
+	}
+}
+
+void SColorGradingPanel::OnWorldAdded(UWorld* World)
+{
+	bRefreshOnNextTick = true;
+}
+
+void SColorGradingPanel::OnWorldDestroyed(UWorld* World)
+{
+	auto NoLongerValid = [World](const FColorGradingListItemRef& ColorGradingItem)
+	{
+		return !ColorGradingItem->Actor.IsValid() || ColorGradingItem->Actor->GetWorld() == World;
+	};
+
+	// Immediately remove any affected actors so we don't try to access them
+	ColorGradingItemList.RemoveAll(NoLongerValid);
+}
+
+void SColorGradingPanel::OnBlueprintCompiled(UBlueprint* Blueprint)
+{
+	Refresh();
+}
+
+void SColorGradingPanel::OnColorGradingDataModelGenerated()
+{
+	FillColorGradingGroupToolBar();
+
+	if (ColorWheelPanel.IsValid())
+	{
+		ColorWheelPanel->Refresh();
+	}
+}
+
+void SColorGradingPanel::OnListSelectionChanged(TSharedRef<SColorGradingObjectList> SourceList, FColorGradingListItemRef SelectedItem, ESelectInfo::Type SelectInfo)
+{
+	if (SelectInfo != ESelectInfo::Direct)
+	{
+		TArray<FColorGradingListItemRef> SelectedObjects = SourceList->GetSelectedItems();
+		TArray<UObject*> ObjectsToColorGrade;
+		for (const FColorGradingListItemRef& SelectedObject : SelectedObjects)
+		{
+			if (SelectedObject->Component.IsValid())
+			{
+				ObjectsToColorGrade.Add(SelectedObject->Component.Get());
+			}
+			else if (SelectedObject->Actor.IsValid())
+			{
+				ObjectsToColorGrade.Add(SelectedObject->Actor.Get());
+			}
+		}
+
+		SetSelectedObjects(ObjectsToColorGrade);
+	}
+}
+
+FReply SColorGradingPanel::DockInLayout()
+{
+	DockCallback.ExecuteIfBound();
+
+	return FReply::Handled();
+}
\ No newline at end of file
diff --git a/Engine/Plugins/Editor/ColorGrading/Source/ColorGradingEditor/Public/ColorGradingListItem.h b/Engine/Plugins/Editor/ColorGrading/Source/ColorGradingEditor/Public/ColorGradingListItem.h
new file mode 100644
--- /dev/null
+++ b/Engine/Plugins/Editor/ColorGrading/Source/ColorGradingEditor/Public/ColorGradingListItem.h
@@ -0,0 +1,85 @@
+// Copyright Epic Games, Inc. All Rights Reserved.
+
+#pragma once
+
+#include "CoreMinimal.h"
+#include "Engine/Brush.h"
+#include "GameFramework/Actor.h"
+#include "ScopedTransaction.h"
+#include "Templates/SharedPointer.h"
+
+typedef TSharedPtr<class IColorGradingListItemGenerator> FColorGradingListItemGeneratorRef;
+typedef TSharedPtr<struct FColorGradingListItem> FColorGradingListItemRef;
+
+DECLARE_DELEGATE_TwoParams(FOnColorGradingItemEnabledChanged, FColorGradingListItemRef, bool);
+
+/** Interface that allows color grading list items to be generated for specific UObjects */
+class COLORGRADINGEDITOR_API IColorGradingListItemGenerator : public TSharedFromThis<IColorGradingListItemGenerator>
+{
+public:
+	virtual ~IColorGradingListItemGenerator() {}
+
+	/** Returns a list of actor classes for which this can generate color grading list items */
+	virtual TArray<TSubclassOf<AActor>> GetActorClassesForListItems() const = 0;
+
+	/** Appends FColorGradingListItems for each color gradable object associated with the given actor */
+	virtual void GenerateColorGradingListItems(AActor* InActor, TArray<FColorGradingListItemRef>& OutList) const = 0;
+};
+
+// Macros to create the lambdas used for a specified UObject's color grading enabled property in the color grading object list
+#define CREATE_IS_ENABLED_LAMBDA(Object, IsEnabledProperty) TAttribute<bool>::CreateLambda([Object]() { return (bool)IsEnabledProperty; })
+#define CREATE_ON_ENABLED_CHANGED_LAMBDA(Object, IsEnabledProperty) FOnColorGradingItemEnabledChanged::CreateLambda([Object](FColorGradingListItemRef ListItem, bool bIsEnabled) \
+	{ \
+		FScopedTransaction Transaction(LOCTEXT("ColorGradingToggledTransaction", "Color Grading Toggled")); \
+		Object->Modify(!Object->IsA<ABrush>()); \
+		IsEnabledProperty = bIsEnabled; \
+	})
+
+/** A structure to store references to color gradable actors and components */
+struct COLORGRADINGEDITOR_API FColorGradingListItem
+{
+public:
+	/** The actor that is color gradable */
+	TWeakObjectPtr<AActor> Actor;
+
+	/** The component that is color gradable */
+	TWeakObjectPtr<UActorComponent> Component;
+
+	/** Attribute that retrieves whether color grading is enabled on the color gradable item */
+	TAttribute<bool> IsItemEnabled;
+
+	/** Delegate raised when the enabled state of the color gradable item has been changed */
+	FOnColorGradingItemEnabledChanged OnItemEnabledChanged;
+
+	FColorGradingListItem(AActor* InActor, UActorComponent* InComponent = nullptr)
+		: Actor(InActor)
+		, Component(InComponent)
+		, IsItemEnabled(false)
+	{ }
+
+	/** Less than operator overload that compares list items alphabetically by their display names */
+	bool operator<(const FColorGradingListItem& Other) const;
+
+public:
+	/** Returns the set of actor classes for which the registered data model generators can generate color grading list items */
+	static const TSet<TSubclassOf<AActor>>& GetActorClassesWithListItemGenerators() { return ActorClassesWithListItemGenerators; }
+
+	/** Registers a new list item generator used to populate a color grading item list */
+	template<class T>
+	static void RegisterColorGradingListItemGenerator()
+	{
+		TSharedRef<IColorGradingListItemGenerator> Generator = MakeShareable(new T());
+		RegisteredListItemGenerators.Add(Generator);
+		ActorClassesWithListItemGenerators.Append(Generator->GetActorClassesForListItems());
+	}
+
+	/** Returns a list of FColorGradingListItems for each color gradable object associated with the given actor */
+	static TArray<FColorGradingListItemRef> GenerateColorGradingListItems(AActor* InActor);
+
+private:
+	/** A list of list item generators that have been registered */
+	static TArray<FColorGradingListItemGeneratorRef> RegisteredListItemGenerators;
+
+	/** The set of actor classes for which the registered data model generators can generate color grading list items */
+	static TSet<TSubclassOf<AActor>> ActorClassesWithListItemGenerators;
+};
\ No newline at end of file
diff --git a/Engine/Plugins/Editor/ColorGrading/Source/ColorGradingEditor/Public/ColorGradingPanelState.h b/Engine/Plugins/Editor/ColorGrading/Source/ColorGradingEditor/Public/ColorGradingPanelState.h
new file mode 100644
--- /dev/null
+++ b/Engine/Plugins/Editor/ColorGrading/Source/ColorGradingEditor/Public/ColorGradingPanelState.h
@@ -0,0 +1,41 @@
+// Copyright Epic Games, Inc. All Rights Reserved.
+
+#pragma once
+
+#include "CoreMinimal.h"
+#include "Styling/SlateTypes.h"
+
+/** Modes that colors can be displayed in for the color grading panel */
+enum class COLORGRADINGEDITOR_API EColorGradingColorDisplayMode
+{
+	RGB,
+	HSV
+};
+
+/**
+ * Stores the state of the color grading panel UI that can be reloaded in cases where the panel or any of its elements
+ * are reloaded (such as when the containing drawer is reopened or docked)
+ */
+struct COLORGRADINGEDITOR_API FColorGradingPanelState
+{
+	/** The objects that are selected in the list */
+	TArray<TWeakObjectPtr<UObject>> SelectedObjects;
+
+	/** The color grading group that is selected */
+	int32 SelectedColorGradingGroup = INDEX_NONE;
+
+	/** The color grading element that is selected */
+	int32 SelectedColorGradingElement = INDEX_NONE;
+
+	/** Indicates which color wheels are hidden */
+	TArray<bool> HiddenColorWheels;
+
+	/** The selected orientation of the color wheels */
+	EOrientation ColorWheelOrientation = EOrientation::Orient_Vertical;
+
+	/** The color display mode of the color wheels */
+	EColorGradingColorDisplayMode ColorDisplayMode;
+
+	/** Indicates which subsections were selected for each section in the details panel */
+	TArray<int32> SelectedDetailsSubsections;
+};
\ No newline at end of file
