Fix missing #undef warning

diff --git a/Engine/Plugins/Editor/ColorGrading/Source/ColorGradingEditor/Private/SColorGradingPanel.cpp b/Engine/Plugins/Editor/ColorGrading/Source/ColorGradingEditor/Private/SColorGradingPanel.cpp
--- a/Engine/Plugins/Editor/ColorGrading/Source/ColorGradingEditor/Private/SColorGradingPanel.cpp
+++ b/Engine/Plugins/Editor/ColorGrading/Source/ColorGradingEditor/Private/SColorGradingPanel.cpp
@@ -763,4 +763,6 @@ FReply SColorGradingPanel::DockInLayout()
 	DockCallback.ExecuteIfBound();
 
 	return FReply::Handled();
-}
\ No newline at end of file
+}
+
+#undef LOCTEXT_NAMESPACE
