Add missing namespace to localized text in ColorGradingListItem macro

diff --git a/Engine/Plugins/Editor/ColorGrading/Source/ColorGradingEditor/Public/ColorGradingListItem.h b/Engine/Plugins/Editor/ColorGrading/Source/ColorGradingEditor/Public/ColorGradingListItem.h
--- a/Engine/Plugins/Editor/ColorGrading/Source/ColorGradingEditor/Public/ColorGradingListItem.h
+++ b/Engine/Plugins/Editor/ColorGrading/Source/ColorGradingEditor/Public/ColorGradingListItem.h
@@ -30,7 +30,7 @@ public:
 #define CREATE_IS_ENABLED_LAMBDA(Object, IsEnabledProperty) TAttribute<bool>::CreateLambda([Object]() { return (bool)IsEnabledProperty; })
 #define CREATE_ON_ENABLED_CHANGED_LAMBDA(Object, IsEnabledProperty) FOnColorGradingItemEnabledChanged::CreateLambda([Object](FColorGradingListItemRef ListItem, bool bIsEnabled) \
 	{ \
-		FScopedTransaction Transaction(LOCTEXT("ColorGradingToggledTransaction", "Color Grading Toggled")); \
+		FScopedTransaction Transaction(NSLOCTEXT("ColorGradingEditor", "ColorGradingToggledTransaction", "Color Grading Toggled")); \
 		Object->Modify(!Object->IsA<ABrush>()); \
 		IsEnabledProperty = bIsEnabled; \
 	})
