[ColorGradingEditor] Fix actor label failing to update when renamed

diff --git a/Engine/Plugins/Editor/ColorGrading/Source/ColorGradingEditor/Private/DataModelGenerators/ColorGradingDataModelGenerator_ColorCorrectRegion.cpp b/Engine/Plugins/Editor/ColorGrading/Source/ColorGradingEditor/Private/DataModelGenerators/ColorGradingDataModelGenerator_ColorCorrectRegion.cpp
--- a/Engine/Plugins/Editor/ColorGrading/Source/ColorGradingEditor/Private/DataModelGenerators/ColorGradingDataModelGenerator_ColorCorrectRegion.cpp
+++ b/Engine/Plugins/Editor/ColorGrading/Source/ColorGradingEditor/Private/DataModelGenerators/ColorGradingDataModelGenerator_ColorCorrectRegion.cpp
@@ -134,7 +134,12 @@ void FColorGradingDataModelGenerator_ColorCorrectRegion::GenerateDataModel(IProp
 			.HAlign(HAlign_Left)
 			[
 				SNew(STextBlock)
-				.Text(FText::FromString(SelectedCCRs[0]->GetActorLabel()))
+				.Text(MakeAttributeLambda([WeakActor = SelectedCCRs[0]]()
+				{
+					return WeakActor.IsValid()
+						? FText::FromString(WeakActor->GetActorLabel())
+						: FText::GetEmpty();
+				}))
 				.Font(FAppStyle::Get().GetFontStyle("NormalFontBold"))
 			];
 
diff --git a/Engine/Plugins/Editor/ColorGrading/Source/ColorGradingEditor/Private/DataModelGenerators/ColorGradingDataModelGenerator_PostProcessVolume.cpp b/Engine/Plugins/Editor/ColorGrading/Source/ColorGradingEditor/Private/DataModelGenerators/ColorGradingDataModelGenerator_PostProcessVolume.cpp
--- a/Engine/Plugins/Editor/ColorGrading/Source/ColorGradingEditor/Private/DataModelGenerators/ColorGradingDataModelGenerator_PostProcessVolume.cpp
+++ b/Engine/Plugins/Editor/ColorGrading/Source/ColorGradingEditor/Private/DataModelGenerators/ColorGradingDataModelGenerator_PostProcessVolume.cpp
@@ -176,7 +176,12 @@ void FColorGradingDataModelGenerator_PostProcessVolume::GenerateDataModel(IPrope
 			.HAlign(HAlign_Left)
 			[
 				SNew(STextBlock)
-				.Text(FText::FromString(SelectedPPVs[0]->GetActorLabel()))
+				.Text(MakeAttributeLambda([WeakActor = SelectedPPVs[0]]()
+				{
+					return WeakActor.IsValid()
+						? FText::FromString(WeakActor->GetActorLabel())
+						: FText::GetEmpty();
+				}))
 				.Font(FAppStyle::Get().GetFontStyle("NormalFontBold"))
 			];
 
