[ColorGrading] Adjust spacing of color grading UI and add alternate layout when vertical space is tight

diff --git a/Engine/Plugins/Editor/ColorGrading/Source/ColorGradingEditor/Private/SColorGradingColorWheel.cpp b/Engine/Plugins/Editor/ColorGrading/Source/ColorGradingEditor/Private/SColorGradingColorWheel.cpp
--- a/Engine/Plugins/Editor/ColorGrading/Source/ColorGradingEditor/Private/SColorGradingColorWheel.cpp
+++ b/Engine/Plugins/Editor/ColorGrading/Source/ColorGradingEditor/Private/SColorGradingColorWheel.cpp
@@ -49,6 +49,8 @@ void SColorGradingColorWheel::Construct(const FArguments& InArgs)
 		.MaxDesiredWidth(400.f)
 		.MinDesiredWidth(400.f);
 
+	HeaderBox = SNew(SBox);
+
 	ChildSlot
 	[
 		SNew(SBox)
@@ -60,7 +62,7 @@ void SColorGradingColorWheel::Construct(const FArguments& InArgs)
 			.AutoHeight()
 			.HAlign(HAlign_Fill)
 			[
-				SAssignNew(HeaderBox, SBox)
+				HeaderBox.ToSharedRef()
 			]
 
 			+ SVerticalBox::Slot()
@@ -74,7 +76,11 @@ void SColorGradingColorWheel::Construct(const FArguments& InArgs)
 			.AutoHeight()
 			.HAlign(HAlign_Center)
 			[
-				ColorSlidersBox.ToSharedRef()
+				SNew(SBox)
+				.Visibility(this, &SColorGradingColorWheel::GetSlidersVisibility)
+				[
+					ColorSlidersBox.ToSharedRef()
+				]
 			]
 		]
 	];
@@ -178,7 +184,7 @@ TSharedRef<SWidget> SColorGradingColorWheel::CreateColorComponentSliders()
 			);
 
 			ColorSlidersVerticalBox->AddSlot()
-			.Padding(FMargin(0.0f, 8.0f, 0.0f, 0.0f))
+			.Padding(FMargin(0.0f, 4.0f, 0.0f, 0.0f))
 			.AutoHeight()
 			.HAlign(HAlign_Fill)
 			.VAlign(VAlign_Fill)
@@ -333,6 +339,13 @@ bool SColorGradingColorWheel::IsPropertyEnabled() const
 	return false;
 }
 
+EVisibility SColorGradingColorWheel::GetSlidersVisibility() const
+{
+	return GetTickSpaceGeometry().GetLocalSize().Y >= 294
+		? EVisibility::Visible
+		: EVisibility::Collapsed;
+}
+
 bool SColorGradingColorWheel::GetColor(FVector4& OutCurrentColor)
 {
 	if (ColorPropertyHandle.IsValid())
diff --git a/Engine/Plugins/Editor/ColorGrading/Source/ColorGradingEditor/Private/SColorGradingColorWheel.h b/Engine/Plugins/Editor/ColorGrading/Source/ColorGradingEditor/Private/SColorGradingColorWheel.h
--- a/Engine/Plugins/Editor/ColorGrading/Source/ColorGradingEditor/Private/SColorGradingColorWheel.h
+++ b/Engine/Plugins/Editor/ColorGrading/Source/ColorGradingEditor/Private/SColorGradingColorWheel.h
@@ -72,6 +72,7 @@ private:
 	FColorPropertyMetadata GetColorPropertyMetadata() const;
 
 	bool IsPropertyEnabled() const;
+	EVisibility GetSlidersVisibility() const;
 
 	bool GetColor(FVector4& OutCurrentColor);
 	void CommitColor(FVector4& NewValue, bool bShouldCommitValueChanges);
diff --git a/Engine/Source/Runtime/AdvancedWidgets/Private/Styling/AdvancedWidgetsStyle.cpp b/Engine/Source/Runtime/AdvancedWidgets/Private/Styling/AdvancedWidgetsStyle.cpp
--- a/Engine/Source/Runtime/AdvancedWidgets/Private/Styling/AdvancedWidgetsStyle.cpp
+++ b/Engine/Source/Runtime/AdvancedWidgets/Private/Styling/AdvancedWidgetsStyle.cpp
@@ -93,6 +93,9 @@ void FAdvancedWidgetsStyle::Create()
 		);
 	}
 
+	Instance->Set("ColorGradingPicker.NumericEntry.TextBox", FCoreStyle::Get().GetWidgetStyle<FEditableTextBoxStyle>("DarkEditableTextBox"));
+	Instance->Set("ColorGrading.NormalFont", FAppStyle::GetFontStyle("PropertyWindow.NormalFont"));
+
 	FLinearColor& VectorColor = ColorSettings.StructColors.FindOrAdd(TBaseStructure<FVector>::Get()->GetStructPathName().ToString());
 	VectorColor = FColorList::Yellow;
 	FLinearColor& RotatorColor = ColorSettings.StructColors.FindOrAdd(TBaseStructure<FRotator>::Get()->GetStructPathName().ToString());
diff --git a/Engine/Source/Runtime/AdvancedWidgets/Private/Widgets/ColorGrading/SColorGradingComponentViewer.cpp b/Engine/Source/Runtime/AdvancedWidgets/Private/Widgets/ColorGrading/SColorGradingComponentViewer.cpp
--- a/Engine/Source/Runtime/AdvancedWidgets/Private/Widgets/ColorGrading/SColorGradingComponentViewer.cpp
+++ b/Engine/Source/Runtime/AdvancedWidgets/Private/Widgets/ColorGrading/SColorGradingComponentViewer.cpp
@@ -33,7 +33,7 @@ void SColorGradingComponentViewer::Construct(const FArguments& InArgs)
 	this->ChildSlot
 	[
 		SNew(SBox)
-		.HeightOverride(22)
+		.HeightOverride(20)
 		.HAlign(HAlign_Fill)
 		.ToolTipText(this, &SColorGradingComponentViewer::GetComponentToolTipText)
 		[
@@ -94,6 +94,7 @@ void SColorGradingComponentViewer::Construct(const FArguments& InArgs)
 	[
 		SAssignNew(NumericEntryBox, SNumericEntryBox<float>)
 		.SpinBoxStyle(&AdvancedWidgetsStyle.GetWidgetStyle<FSpinBoxStyle>("ColorGradingComponentViewer.NumericEntry"))
+		.Font(AdvancedWidgetsStyle.GetFontStyle("ColorGrading.NormalFont"))
 		.EditableTextBoxStyle(&AdvancedWidgetsStyle.GetWidgetStyle<FEditableTextBoxStyle>("ColorGradingComponentViewer.NumericEntry.TextBox"))
 		.UndeterminedString(NSLOCTEXT("PropertyEditor", "MultipleValues", "Multiple Values"))
 		.Value(InArgs._Value)
diff --git a/Engine/Source/Runtime/AdvancedWidgets/Private/Widgets/ColorGrading/SColorGradingPicker.cpp b/Engine/Source/Runtime/AdvancedWidgets/Private/Widgets/ColorGrading/SColorGradingPicker.cpp
--- a/Engine/Source/Runtime/AdvancedWidgets/Private/Widgets/ColorGrading/SColorGradingPicker.cpp
+++ b/Engine/Source/Runtime/AdvancedWidgets/Private/Widgets/ColorGrading/SColorGradingPicker.cpp
@@ -1,5 +1,8 @@
 // Copyright Epic Games, Inc. All Rights Reserved.
+
 #include "Widgets/ColorGrading/SColorGradingPicker.h"
+
+#include "Styling/AdvancedWidgetsStyle.h"
 #include "Widgets/ColorGrading/SColorGradingWheel.h"
 #include "Widgets/Input/SNumericEntryBox.h"
 #include "Widgets/Layout/SBox.h"
@@ -54,6 +57,8 @@ void SColorGradingPicker::Construct(const FArguments& InArgs)
 	ExternalBeginMouseCaptureDelegate = InArgs._OnBeginMouseCapture;
 	ExternalEndMouseCaptureDelegate = InArgs._OnEndMouseCapture;
 
+	const auto& AdvancedWidgetsStyle = UE::AdvancedWidgets::FAdvancedWidgetsStyle::Get();
+
 	ChildSlot
 		[
 			SNew(SVerticalBox)
@@ -85,7 +90,8 @@ void SColorGradingPicker::Construct(const FArguments& InArgs)
 						.MinDesiredWidth(400.f)
 						[
 							SAssignNew(NumericEntryBoxWidget, SNumericEntryBox<float>)
-								.EditableTextBoxStyle(&FCoreStyle::Get().GetWidgetStyle<FEditableTextBoxStyle>("DarkEditableTextBox"))
+								.EditableTextBoxStyle(&AdvancedWidgetsStyle.GetWidgetStyle<FEditableTextBoxStyle>("ColorGradingPicker.NumericEntry.TextBox"))
+								.Font(AdvancedWidgetsStyle.GetFontStyle("ColorGrading.NormalFont"))
 								.Value(this, &SColorGradingPicker::OnGetMainValue)
 								.OnValueCommitted(this, &SColorGradingPicker::OnMainValueCommitted)
 								.OnValueChanged(this, &SColorGradingPicker::OnMainValueChanged, false)
diff --git a/Engine/Source/Runtime/AdvancedWidgets/Private/Widgets/ColorGrading/SColorGradingWheel.cpp b/Engine/Source/Runtime/AdvancedWidgets/Private/Widgets/ColorGrading/SColorGradingWheel.cpp
--- a/Engine/Source/Runtime/AdvancedWidgets/Private/Widgets/ColorGrading/SColorGradingWheel.cpp
+++ b/Engine/Source/Runtime/AdvancedWidgets/Private/Widgets/ColorGrading/SColorGradingWheel.cpp
@@ -182,13 +182,7 @@ int32 SColorGradingWheel::OnPaint(const FPaintArgs& Args, const FGeometry& Allot
 		AllottedGeometrySize = AllottedGeometry.GetLocalSize();
 	}
 
-	const float SelectorShrinkThreshold = 150.f;
-	FVector2f SelectorSize = SelectorImage->ImageSize;
-	if (AllottedGeometrySize.X < SelectorShrinkThreshold || AllottedGeometrySize.Y < SelectorShrinkThreshold)
-	{
-		SelectorSize *= 0.75f;
-	}
-
+	const FVector2f SelectorSize = SelectorImage->ImageSize;
 	const FVector2f CircleSize = AllottedGeometrySize - SelectorSize;
 
 	FSlateDrawElement::MakeBox(
