[ColorGrading] Reorder color wheels to match Details panel

diff --git a/Engine/Plugins/Editor/ColorGrading/Source/ColorGradingEditor/Private/SColorGradingColorWheelPanel.cpp b/Engine/Plugins/Editor/ColorGrading/Source/ColorGradingEditor/Private/SColorGradingColorWheelPanel.cpp
--- a/Engine/Plugins/Editor/ColorGrading/Source/ColorGradingEditor/Private/SColorGradingColorWheelPanel.cpp
+++ b/Engine/Plugins/Editor/ColorGrading/Source/ColorGradingEditor/Private/SColorGradingColorWheelPanel.cpp
@@ -358,11 +358,11 @@ void SColorGradingColorWheelPanel::FillColorWheels(const FColorGradingEditorData
 		}
 	};
 
-	FillColorWheel(ColorWheels[0], ColorGradingElement.OffsetPropertyHandle);
-	FillColorWheel(ColorWheels[1], ColorGradingElement.GammaPropertyHandle);
-	FillColorWheel(ColorWheels[2], ColorGradingElement.GainPropertyHandle);
-	FillColorWheel(ColorWheels[3], ColorGradingElement.SaturationPropertyHandle);
-	FillColorWheel(ColorWheels[4], ColorGradingElement.ContrastPropertyHandle);
+	FillColorWheel(ColorWheels[0], ColorGradingElement.SaturationPropertyHandle);
+	FillColorWheel(ColorWheels[1], ColorGradingElement.ContrastPropertyHandle);
+	FillColorWheel(ColorWheels[2], ColorGradingElement.GammaPropertyHandle);
+	FillColorWheel(ColorWheels[3], ColorGradingElement.GainPropertyHandle);
+	FillColorWheel(ColorWheels[4], ColorGradingElement.OffsetPropertyHandle);
 }
 
 void SColorGradingColorWheelPanel::ClearColorWheels()
