[ColorGrading] Center color wheel sliders with wheel when there's too much vertical space

diff --git a/Engine/Plugins/Editor/ColorGrading/Source/ColorGradingEditor/Private/SColorGradingColorWheel.cpp b/Engine/Plugins/Editor/ColorGrading/Source/ColorGradingEditor/Private/SColorGradingColorWheel.cpp
--- a/Engine/Plugins/Editor/ColorGrading/Source/ColorGradingEditor/Private/SColorGradingColorWheel.cpp
+++ b/Engine/Plugins/Editor/ColorGrading/Source/ColorGradingEditor/Private/SColorGradingColorWheel.cpp
@@ -54,7 +54,7 @@ void SColorGradingColorWheel::Construct(const FArguments& InArgs)
 	ChildSlot
 	[
 		SNew(SBox)
-		.Padding(16.f, 8.f)
+		.Padding(ColumnPadding.X, ColumnPadding.Y)
 		[
 			SNew(SVerticalBox)
 
@@ -66,20 +66,55 @@ void SColorGradingColorWheel::Construct(const FArguments& InArgs)
 			]
 
 			+ SVerticalBox::Slot()
-			.FillHeight(1.0f)
+			.FillHeight(1.0)
 			.HAlign(HAlign_Fill)
 			[
-				ColorPickerBox.ToSharedRef()
-			]
+				SNew(SOverlay)
 
-			+ SVerticalBox::Slot()
-			.AutoHeight()
-			.HAlign(HAlign_Center)
-			[
-				SNew(SBox)
-				.Visibility(this, &SColorGradingColorWheel::GetSlidersVisibility)
+				// Standard/"short" layout
+				+ SOverlay::Slot()
+				.HAlign(HAlign_Fill)
+				.VAlign(VAlign_Fill)
 				[
-					ColorSlidersBox.ToSharedRef()
+					SNew(SVerticalBox)
+					.Visibility(this, &SColorGradingColorWheel::GetShortLayoutVisibility)
+
+					+ SVerticalBox::Slot()
+					.FillHeight(1.f)
+					[
+						ColorPickerBox.ToSharedRef()
+					]
+
+					+ SVerticalBox::Slot()
+					.AutoHeight()
+					[
+						SNew(SBox)
+						.Visibility(this, &SColorGradingColorWheel::GetSlidersVisibility)
+						[
+							ColorSlidersBox.ToSharedRef()
+						]
+					]
+				]
+
+				// Tall layout
+				+ SOverlay::Slot()
+				.HAlign(HAlign_Fill)
+				.VAlign(VAlign_Center)
+				[
+					SNew(SVerticalBox)
+					.Visibility(this, &SColorGradingColorWheel::GetTallLayoutVisibility)
+
+					+ SVerticalBox::Slot()
+					.AutoHeight()
+					[
+						ColorPickerBox.ToSharedRef()
+					]
+
+					+ SVerticalBox::Slot()
+					.AutoHeight()
+					[
+						ColorSlidersBox.ToSharedRef()
+					]
 				]
 			]
 		]
@@ -147,6 +182,7 @@ TSharedRef<SWidget> SColorGradingColorWheel::SColorGradingColorWheel::CreateColo
 	if (ColorPropertyHandle.IsValid() && ColorPropertyHandle.Pin()->IsValidHandle())
 	{
 		return SNew(UE::ColorGrading::SColorGradingPicker)
+			.DesiredWheelSize(this, &SColorGradingColorWheel::GetMaxWheelWidth)
 			.ValueMin(GetMetadataMinValue())
 			.ValueMax(GetMetadataMaxValue())
 			.SliderValueMin(GetMetadataSliderMinValue())
@@ -346,6 +382,33 @@ EVisibility SColorGradingColorWheel::GetSlidersVisibility() const
 		: EVisibility::Collapsed;
 }
 
+EVisibility SColorGradingColorWheel::GetShortLayoutVisibility() const
+{
+	return ShouldUseTallLayout()
+		? EVisibility::Collapsed
+		: EVisibility::Visible;
+}
+
+EVisibility SColorGradingColorWheel::GetTallLayoutVisibility() const
+{
+	return ShouldUseTallLayout()
+		? EVisibility::Visible
+		: EVisibility::Collapsed;
+}
+
+int32 SColorGradingColorWheel::GetMaxWheelWidth() const
+{
+	const float InnerWidth = GetTickSpaceGeometry().GetLocalSize().X - ColumnPadding.X * 2;
+	return FMath::FloorToInt32(FMath::Max(InnerWidth, 0.f));
+}
+
+bool SColorGradingColorWheel::ShouldUseTallLayout() const
+{
+	// Switch to tall layout when the column space not occupied by the wheel is above this size
+	const float RemainingSpaceThrehold = 450;
+	return GetTickSpaceGeometry().GetLocalSize().Y >= (RemainingSpaceThrehold + GetMaxWheelWidth());
+}
+
 bool SColorGradingColorWheel::GetColor(FVector4& OutCurrentColor)
 {
 	if (ColorPropertyHandle.IsValid())
diff --git a/Engine/Plugins/Editor/ColorGrading/Source/ColorGradingEditor/Private/SColorGradingColorWheel.h b/Engine/Plugins/Editor/ColorGrading/Source/ColorGradingEditor/Private/SColorGradingColorWheel.h
--- a/Engine/Plugins/Editor/ColorGrading/Source/ColorGradingEditor/Private/SColorGradingColorWheel.h
+++ b/Engine/Plugins/Editor/ColorGrading/Source/ColorGradingEditor/Private/SColorGradingColorWheel.h
@@ -73,6 +73,10 @@ private:
 
 	bool IsPropertyEnabled() const;
 	EVisibility GetSlidersVisibility() const;
+	EVisibility GetShortLayoutVisibility() const;
+	EVisibility GetTallLayoutVisibility() const;
+	int32 GetMaxWheelWidth() const;
+	bool ShouldUseTallLayout() const;
 
 	bool GetColor(FVector4& OutCurrentColor);
 	void CommitColor(FVector4& NewValue, bool bShouldCommitValueChanges);
@@ -108,6 +112,9 @@ private:
 	bool GetMetadataSupportDynamicSliderMaxValue() const;
 
 private:
+	/** Padding applied to the whole column */
+	const FVector2f ColumnPadding = FVector2f(16.f, 8.f);
+
 	TSharedPtr<UE::ColorGrading::SColorGradingPicker> ColorGradingPicker;
 	TSharedPtr<SBox> HeaderBox;
 	TSharedPtr<SBox> ColorPickerBox;
diff --git a/Engine/Source/Runtime/AdvancedWidgets/Private/Widgets/ColorGrading/SColorGradingPicker.cpp b/Engine/Source/Runtime/AdvancedWidgets/Private/Widgets/ColorGrading/SColorGradingPicker.cpp
--- a/Engine/Source/Runtime/AdvancedWidgets/Private/Widgets/ColorGrading/SColorGradingPicker.cpp
+++ b/Engine/Source/Runtime/AdvancedWidgets/Private/Widgets/ColorGrading/SColorGradingPicker.cpp
@@ -73,6 +73,7 @@ void SColorGradingPicker::Construct(const FArguments& InArgs)
 						.MaxAspectRatio(1.0f)
 						[
 							SNew(SColorGradingWheel)
+								.DesiredWheelSize(InArgs._DesiredWheelSize)
 								.SelectedColor(TAttribute<FLinearColor>::Create(TAttribute<FLinearColor>::FGetter::CreateSP(this, &SColorGradingPicker::GetCurrentLinearColor)))
 								.ExponentDisplacement(ColorGradingWheelExponent)
 								.OnValueChanged(this, &SColorGradingPicker::HandleCurrentColorValueChanged, false)
diff --git a/Engine/Source/Runtime/AdvancedWidgets/Private/Widgets/ColorGrading/SColorGradingWheel.cpp b/Engine/Source/Runtime/AdvancedWidgets/Private/Widgets/ColorGrading/SColorGradingWheel.cpp
--- a/Engine/Source/Runtime/AdvancedWidgets/Private/Widgets/ColorGrading/SColorGradingWheel.cpp
+++ b/Engine/Source/Runtime/AdvancedWidgets/Private/Widgets/ColorGrading/SColorGradingWheel.cpp
@@ -171,15 +171,15 @@ int32 SColorGradingWheel::OnPaint(const FPaintArgs& Args, const FGeometry& Allot
 	const bool bIsEnabled = ShouldBeEnabled(bParentEnabled);
 	const ESlateDrawEffect DrawEffects = bIsEnabled ? ESlateDrawEffect::None : ESlateDrawEffect::DisabledEffect;
 
-	FVector2f AllottedGeometrySize;
+	FVector2f AllottedGeometrySize = AllottedGeometry.GetLocalSize();
+
 	if (bIsAttributeDesiredWheelSizeSet)
 	{
+		// Even if a desired size is provided, make sure the wheel is painted within the allotted geometry
 		int32 CachedDesiredWheelSize = DesiredWheelSizeAttribute.Get();
-		AllottedGeometrySize = FVector2f(CachedDesiredWheelSize, CachedDesiredWheelSize);
-	}
-	else
-	{
-		AllottedGeometrySize = AllottedGeometry.GetLocalSize();
+		float ActualSize = FMath::Min(CachedDesiredWheelSize, AllottedGeometrySize.GetMin());
+
+		AllottedGeometrySize = FVector2f(ActualSize, ActualSize);
 	}
 
 	const FVector2f SelectorSize = SelectorImage->ImageSize;
diff --git a/Engine/Source/Runtime/AdvancedWidgets/Public/Widgets/ColorGrading/SColorGradingPicker.h b/Engine/Source/Runtime/AdvancedWidgets/Public/Widgets/ColorGrading/SColorGradingPicker.h
--- a/Engine/Source/Runtime/AdvancedWidgets/Public/Widgets/ColorGrading/SColorGradingPicker.h
+++ b/Engine/Source/Runtime/AdvancedWidgets/Public/Widgets/ColorGrading/SColorGradingPicker.h
@@ -35,6 +35,7 @@ public:
 		: _AllowSpin(true)
 		, _SupportDynamicSliderMaxValue(false)
 		, _SupportDynamicSliderMinValue(false)
+		, _DesiredWheelSize()
 		, _MainDelta(0.01f)
 		, _MainShiftMultiplier(10.f)
 		, _MainCtrlMultiplier(0.1f)
@@ -54,6 +55,8 @@ public:
 		/** Tell us if we want to support dynamically changing of the min value using ctrl */
 		SLATE_ATTRIBUTE(bool, SupportDynamicSliderMinValue)
 
+		SLATE_ATTRIBUTE(int32, DesiredWheelSize)
+
 		SLATE_ARGUMENT(float, MainDelta)
 
 		/** Multiplier to use when shift is held down */
