[ColorGradingEditor] Fix crash when modifying a property where the outer object is null

diff --git a/Engine/Plugins/Editor/ColorGrading/Source/ColorGradingEditor/Private/SColorGradingColorWheel.cpp b/Engine/Plugins/Editor/ColorGrading/Source/ColorGradingEditor/Private/SColorGradingColorWheel.cpp
--- a/Engine/Plugins/Editor/ColorGrading/Source/ColorGradingEditor/Private/SColorGradingColorWheel.cpp
+++ b/Engine/Plugins/Editor/ColorGrading/Source/ColorGradingEditor/Private/SColorGradingColorWheel.cpp
@@ -465,6 +465,11 @@ void SColorGradingColorWheel::TransactColorValue()
 		ColorPropertyHandle.GetHandle()->GetOuterObjects(OuterObjects);
 		for (UObject* Object : OuterObjects)
 		{
+			if (!Object)
+			{
+				continue;
+			}
+
 			if (!Object->HasAnyFlags(RF_Transactional))
 			{
 				Object->SetFlags(RF_Transactional);
