Color Grading x Multi User: Adjust behaviour of the "Add property to replication" property row extension button:
- Only show "Add to replication" details extension button if that object is in an actor hierarchy that is replicated.
- Adjust Color Grading panel UI so more than 1 extension button can fit. This sometimes misaligns the icons when the panel does not have enoug space due to the property editor's SClippingHorizontalBox taking up space even when it clips its elements. That will be fixed for 5.7.

diff --git a/Engine/Plugins/Developer/Concert/ConcertApp/MultiUserClient/Source/MultiUserClient/Private/Replication/Misc/DetailRowExtensionHandler.cpp b/Engine/Plugins/Developer/Concert/ConcertApp/MultiUserClient/Source/MultiUserClient/Private/Replication/Misc/DetailRowExtensionHandler.cpp
--- a/Engine/Plugins/Developer/Concert/ConcertApp/MultiUserClient/Source/MultiUserClient/Private/Replication/Misc/DetailRowExtensionHandler.cpp
+++ b/Engine/Plugins/Developer/Concert/ConcertApp/MultiUserClient/Source/MultiUserClient/Private/Replication/Misc/DetailRowExtensionHandler.cpp
@@ -140,6 +140,36 @@ namespace UE::MultiUserClient::Replication::DetailRowExtensionHandler
 		const bool bIsConnected = NewState == EMultiUserReplicationConnectionState::Connected;
 		SetMinimumRightColumnWidth(bIsConnected);
 	}
+
+	static TArray<UObject*> GetObjects(const TSharedPtr<IPropertyHandle>& InPropertyHandle)
+	{
+		TArray<UObject*> Result;
+
+		if (InPropertyHandle)
+		{
+			InPropertyHandle->GetOuterObjects(Result);
+		}
+		
+		return Result;
+	}
+
+	static bool IsObjectHierarchyReplicated(
+		const TWeakPtr<FMultiUserReplicationManager>& InWeakReplicationManager, TConstArrayView<UObject*> InObjects)
+	{
+		const TSharedPtr<FMultiUserReplicationManager> ReplicationManager = InWeakReplicationManager.Pin();
+		const FOnlineClientManager* OnlineClientManager = ReplicationManager ? ReplicationManager->GetOnlineClientManager() : nullptr;
+		if (!OnlineClientManager // Unset when not in any replicated session.
+			|| InObjects.IsEmpty()) 
+		{
+			return false;
+		}
+
+		return Algo::AnyOf(InObjects, [OnlineClientManager](UObject* Object)
+		{
+			AActor* OwningActor = Object->IsA<AActor>() ? Cast<AActor>(Object) : Object->GetTypedOuter<AActor>();
+			return OwningActor && OnlineClientManager->GetAuthorityCache().IsObjectOrChildReferenced(OwningActor);
+		});
+	}
 }
 
 namespace UE::MultiUserClient::Replication
@@ -188,6 +218,16 @@ namespace UE::MultiUserClient::Replication
 		{
 			return;
 		}
+
+		// Only add the buttons if the object or one of its outers is being replicated. 
+		// Adding them to more than that is overkill and can even interfere with details panels where it is inappropriate (Blueprint Editor, etc.).
+		const bool bIsReplicated = DetailRowExtensionHandler::IsObjectHierarchyReplicated(
+			WeakReplicationManager, DetailRowExtensionHandler::GetObjects(Args.PropertyHandle)
+			);
+		if (!bIsReplicated)
+		{
+			return;
+		}
 		
 		FPropertyRowExtensionButton& CreateKey = OutExtensionButtons.AddDefaulted_GetRef();
 		CreateKey.Icon = FSlateIcon(FMultiUserReplicationStyle::GetStyleSetName(), TEXT("MultiUser.Icons.AddProperty.Small"));
diff --git a/Engine/Plugins/Editor/ColorGrading/Source/ColorGradingEditor/Private/SColorGradingColorWheelPanel.cpp b/Engine/Plugins/Editor/ColorGrading/Source/ColorGradingEditor/Private/SColorGradingColorWheelPanel.cpp
--- a/Engine/Plugins/Editor/ColorGrading/Source/ColorGradingEditor/Private/SColorGradingColorWheelPanel.cpp
+++ b/Engine/Plugins/Editor/ColorGrading/Source/ColorGradingEditor/Private/SColorGradingColorWheelPanel.cpp
@@ -408,14 +408,13 @@ TSharedRef<SWidget> SColorGradingColorWheelPanel::CreateColorWheelHeaderWidget(c
 				[
 					NodeWidgets.NameWidget.ToSharedRef()
 				];
-
 			PropertyNameBox->AddSlot()
 				.HAlign(HAlign_Right)
 				.VAlign(VAlign_Center)
 				.FillWidth(1.0f)
 				[
 					SNew(SBox)
-					.WidthOverride(22.0f)
+					.MinDesiredWidth(22.0f)
 					[
 						CreateColorPropertyExtensions(ColorPropertyHandle, TreeNode, ResetToDefaultOverride)
 					]
@@ -432,12 +431,10 @@ TSharedRef<SWidget> SColorGradingColorWheelPanel::CreateColorPropertyExtensions(
 	const TSharedPtr<IDetailTreeNode>& DetailTreeNode,
 	const TOptional<FResetToDefaultOverride>& ResetToDefaultOverride)
 {
-	TArray<FPropertyRowExtensionButton> ExtensionButtons;
-
 	// Use a weak pointer to pass into delegates
 	TWeakPtr<IPropertyHandle> WeakColorPropertyHandle = ColorPropertyHandle;
 
-	FPropertyRowExtensionButton& ResetToDefaultButton = ExtensionButtons.AddDefaulted_GetRef();
+	FPropertyRowExtensionButton ResetToDefaultButton;
 	ResetToDefaultButton.Label = NSLOCTEXT("PropertyEditor", "ResetToDefault", "Reset to Default");
 	ResetToDefaultButton.UIAction = FUIAction(
 		FExecuteAction::CreateLambda([WeakColorPropertyHandle, ResetToDefaultOverride]
@@ -485,6 +482,7 @@ TSharedRef<SWidget> SColorGradingColorWheelPanel::CreateColorPropertyExtensions(
 	Args.OwnerTreeNode = DetailTreeNode;
 	Args.PropertyHandle = ColorPropertyHandle;
 
+	TArray<FPropertyRowExtensionButton> ExtensionButtons;
 	PropertyEditorModule.GetGlobalRowExtensionDelegate().Broadcast(Args, ExtensionButtons);
 
 	FSlimHorizontalToolBarBuilder ToolbarBuilder(TSharedPtr<FUICommandList>(), FMultiBoxCustomization::None);
@@ -492,10 +490,21 @@ TSharedRef<SWidget> SColorGradingColorWheelPanel::CreateColorPropertyExtensions(
 	ToolbarBuilder.SetStyle(&FAppStyle::Get(), "DetailsView.ExtensionToolBar");
 	ToolbarBuilder.SetIsFocusable(false);
 
+	// Always show reset to default. The other button are shown if there is enough space.
 	for (const FPropertyRowExtensionButton& Extension : ExtensionButtons)
 	{
-		ToolbarBuilder.AddToolBarButton(Extension.UIAction, NAME_None, Extension.Label, Extension.ToolTip, Extension.Icon);
+		ToolbarBuilder.AddToolBarButton(
+			Extension.UIAction, NAME_None, Extension.Label, Extension.ToolTip, Extension.Icon,
+			EUserInterfaceActionType::Button, NAME_None, {}, {},
+			FMenuEntryResizeParams{ .AllowClipping = true }
+			);
 	}
+	// Add the reset button last so it's always the right-most widget.
+	ToolbarBuilder.AddToolBarButton(
+		ResetToDefaultButton.UIAction, NAME_None, ResetToDefaultButton.Label, ResetToDefaultButton.ToolTip, ResetToDefaultButton.Icon,
+		EUserInterfaceActionType::Button, NAME_None, {}, {},
+		FMenuEntryResizeParams{ .AllowClipping = false }
+		);
 
 	return ToolbarBuilder.MakeWidget();
 }
